Copying //UE4/Dev-Sequencer to //UE4/Main (Source: //UE4/Dev-Sequencer @ 3617642)

diff --git a/Engine/Plugins/Media/ImgMedia/Source/ImgMedia/Private/Loader/ImgMediaLoader.cpp b/Engine/Plugins/Media/ImgMedia/Source/ImgMedia/Private/Loader/ImgMediaLoader.cpp
new file mode 100644
--- /dev/null
+++ b/Engine/Plugins/Media/ImgMedia/Source/ImgMedia/Private/Loader/ImgMediaLoader.cpp
@@ -0,0 +1,464 @@
+// Copyright 1998-2016 Epic Games, Inc. All Rights Reserved.
+
+#include "ImgMediaLoader.h"
+#include "ImgMediaPrivate.h"
+
+#include "HAL/FileManager.h"
+#include "HAL/PlatformProcess.h"
+#include "Misc/Paths.h"
+#include "Misc/QueuedThreadPool.h"
+#include "Misc/ScopeLock.h"
+#include "Modules/ModuleManager.h"
+#include "UObject/Class.h"
+#include "UObject/UObjectGlobals.h"
+
+#include "GenericImgMediaReader.h"
+#include "IImgMediaReader.h"
+#include "ImgMediaLoaderWork.h"
+#include "ImgMediaTextureSample.h"
+
+#if IMGMEDIA_EXR_SUPPORTED_PLATFORM
+	#include "ExrImgMediaReader.h"
+#endif
+
+
+/* FImgMediaLoader structors
+ *****************************************************************************/
+
+FImgMediaLoader::FImgMediaLoader()
+	: Frames(1)
+	, ImageWrapperModule(FModuleManager::LoadModuleChecked<IImageWrapperModule>("ImageWrapper"))
+	, Initialized(false)
+	, NumLoadAhead(0)
+	, NumLoadBehind(0)
+	, SequenceDim(FIntPoint::ZeroValue)
+	, SequenceDuration(FTimespan::Zero())
+	, SequenceFps(0.0f)
+	, LastRequestedFrame(INDEX_NONE)
+{ }
+
+
+FImgMediaLoader::~FImgMediaLoader()
+{
+	FScopeLock Lock(&CriticalSection);
+
+	// clear pending frames
+	PendingFrameNumbers.Empty();
+
+	// abandon queued work items
+	for (auto WorkPair : QueuedWorks)
+	{
+		auto QueuedWork = WorkPair.Value;
+		check(!WorkPool.Contains(QueuedWork));
+
+		if (GThreadPool->RetractQueuedWork(QueuedWork))
+		{
+			delete QueuedWork;
+		}
+		else
+		{
+			QueuedWork->DeleteWhenDone();
+		}
+	}
+
+	QueuedWorks.Empty();
+
+	// release work item pool
+	for (auto Work : WorkPool)
+	{
+		delete Work;
+	}
+
+	WorkPool.Empty();
+}
+
+
+/* FImgMediaLoader interface
+ *****************************************************************************/
+
+uint64 FImgMediaLoader::GetBitRate() const
+{
+	FScopeLock Lock(&CriticalSection);
+	return SequenceDim.X * SequenceDim.Y * sizeof(uint16) * SequenceFps * 8;
+}
+
+
+void FImgMediaLoader::GetBusyTimeRanges(TRangeSet<FTimespan>& OutRangeSet) const
+{
+	FScopeLock Lock(&CriticalSection);
+	FrameNumbersToTimeRanges(QueuedFrameNumbers, OutRangeSet);
+}
+
+
+void FImgMediaLoader::GetCompletedTimeRanges(TRangeSet<FTimespan>& OutRangeSet) const
+{
+	FScopeLock Lock(&CriticalSection);
+
+	TArray<int32> CompletedFrames;
+	Frames.GetKeys(CompletedFrames);
+	FrameNumbersToTimeRanges(CompletedFrames, OutRangeSet);
+}
+
+
+TSharedPtr<FImgMediaTextureSample, ESPMode::ThreadSafe> FImgMediaLoader::GetFrameSample(FTimespan Time)
+{
+	const int32 FrameIndex = TimeToFrame(Time);
+
+	if (FrameIndex == INDEX_NONE)
+	{
+		return nullptr;
+	}
+
+	FScopeLock ScopeLock(&CriticalSection);
+
+	const TSharedPtr<FImgMediaFrame, ESPMode::ThreadSafe>* Frame = Frames.FindAndTouch(FrameIndex);
+
+	if (Frame == nullptr)
+	{
+		return nullptr;
+	}
+
+	const FTimespan FrameTime = FTimespan::FromSeconds(FrameIndex / SequenceFps);
+	const FTimespan FrameDuration = FTimespan::FromSeconds(1.0f / SequenceFps);
+
+	auto Sample = MakeShared<FImgMediaTextureSample, ESPMode::ThreadSafe>();
+	
+	if (!Sample->Initialize(*Frame->Get(), SequenceDim, FrameTime, FrameDuration))
+	{
+		return nullptr;
+	}
+
+	return Sample;
+}
+
+
+void FImgMediaLoader::GetPendingTimeRanges(TRangeSet<FTimespan>& OutRangeSet) const
+{
+	FScopeLock Lock(&CriticalSection);
+	FrameNumbersToTimeRanges(PendingFrameNumbers, OutRangeSet);
+}
+
+
+void FImgMediaLoader::Initialize(const FString& SequencePath, const float FpsOverride)
+{
+	check(!Initialized); // reinitialization not allowed for now
+
+	if (SequencePath.IsEmpty())
+	{
+		return;
+	}
+
+	// locate image sequence files
+	TArray<FString> FoundFiles;
+	IFileManager::Get().FindFiles(FoundFiles, *SequencePath, TEXT("*"));
+
+	if (FoundFiles.Num() == 0)
+	{
+		UE_LOG(LogImgMedia, Error, TEXT("The directory %s does not contain any image files"), *SequencePath);
+		return;
+	}
+
+	UE_LOG(LogImgMedia, Verbose, TEXT("Found %i image files in %s"), FoundFiles.Num(), *SequencePath);
+
+	FoundFiles.Sort();
+
+	for (const auto& File : FoundFiles)
+	{
+		ImagePaths.Add(FPaths::Combine(SequencePath, File));
+	}
+
+	FScopeLock Lock(&CriticalSection);
+
+	// create image reader
+	const FString FirstExtension = FPaths::GetExtension(ImagePaths[0]);
+
+	if (FirstExtension == TEXT("exr"))
+	{
+#if IMGMEDIA_EXR_SUPPORTED_PLATFORM
+		Reader = MakeShareable(new FExrImgMediaReader);
+#else
+		UE_LOG(LogImgMedia, Error, TEXT("EXR image sequences are currently supported on macOS and Windows only"));
+		return;
+#endif
+	}
+	else
+	{
+		Reader = MakeShareable(new FGenericImgMediaReader(ImageWrapperModule));
+	}
+
+	// fetch sequence attributes from first image
+	FImgMediaFrameInfo FirstFrameInfo;
+
+	if (!Reader->GetFrameInfo(ImagePaths[0], FirstFrameInfo))
+	{
+		UE_LOG(LogImgMedia, Error, TEXT("Failed to get frame information from first image in %s"), *SequencePath);
+		return;
+	}
+
+	if (FirstFrameInfo.UncompressedSize == 0)
+	{
+		UE_LOG(LogImgMedia, Error, TEXT("The first image in sequence %s does not have a valid frame size"), *SequencePath);
+		return;
+	}
+
+	if (FirstFrameInfo.Dim.GetMin() <= 0)
+	{
+		UE_LOG(LogImgMedia, Error, TEXT("The first image in sequence %s does not have a valid dimension"), *SequencePath);
+		return;
+	}
+
+	SequenceDim = FirstFrameInfo.Dim;
+
+	if (FpsOverride > 0.0f)
+	{
+		SequenceFps = FpsOverride;
+	}
+	else
+	{
+		SequenceFps = FirstFrameInfo.Fps;
+	}
+
+	SequenceDuration = FTimespan::FromSeconds(ImagePaths.Num() / SequenceFps);
+
+	// initialize loader
+	auto Settings = GetDefault<UImgMediaSettings>();
+
+	const FPlatformMemoryStats Stats = FPlatformMemory::GetStats();
+	const SIZE_T DesiredCacheSize = Settings->CacheSizeGB * 1024 * 1024 * 1024;
+	const SIZE_T CacheSize = FMath::Clamp(DesiredCacheSize, (SIZE_T)0, (SIZE_T)Stats.AvailablePhysical);
+
+	const int32 MaxFramesToLoad = (int32)(CacheSize / FirstFrameInfo.UncompressedSize);
+	const int32 NumFramesToLoad = FMath::Clamp(MaxFramesToLoad, 0, ImagePaths.Num());
+	const float LoadBehindScale = FMath::Clamp(Settings->CacheBehindPercentage, 0.0f, 100.0f) / 100.0f;
+
+	NumLoadBehind = (int32)(LoadBehindScale * NumFramesToLoad);
+	NumLoadAhead = NumFramesToLoad - NumLoadBehind;
+
+	Frames.Empty(NumFramesToLoad);
+	Update(0, 0.0f);
+
+	// update info
+	Info = TEXT("Image Sequence\n");
+	Info += FString::Printf(TEXT("    Dimension: %i x %i\n"), SequenceDim.X, SequenceDim.Y);
+	Info += FString::Printf(TEXT("    Format: %s\n"), *FirstFrameInfo.FormatName);
+	Info += FString::Printf(TEXT("    Compression: %s\n"), *FirstFrameInfo.CompressionName);
+	Info += FString::Printf(TEXT("    Frames: %i\n"), ImagePaths.Num());
+	Info += FString::Printf(TEXT("    FPS: %f\n"), SequenceFps);
+
+	FPlatformMisc::MemoryBarrier();
+
+	Initialized = true;
+}
+
+
+bool FImgMediaLoader::RequestFrame(FTimespan Time, float PlayRate)
+{
+	const int32 FrameIndex = TimeToFrame(Time);
+
+	if ((FrameIndex == INDEX_NONE) || (FrameIndex == LastRequestedFrame))
+	{
+		return false;
+	}
+
+	Update(FrameIndex, PlayRate);
+	LastRequestedFrame = FrameIndex;
+
+	return true;
+}
+
+
+/* FImgMediaLoader implementation
+ *****************************************************************************/
+
+void FImgMediaLoader::FrameNumbersToTimeRanges(const TArray<int32>& FrameNumbers, TRangeSet<FTimespan>& OutRangeSet) const
+{
+	if (SequenceFps <= 0.0f)
+	{
+		return;
+	}
+
+	const FTimespan FrameDuration = FTimespan::FromSeconds(1.0 / SequenceFps);
+
+	for (const auto Frame : FrameNumbers)
+	{
+		const FTimespan StartTime = FTimespan::FromSeconds(Frame / SequenceFps);
+		OutRangeSet.Add(TRange<FTimespan>(StartTime, StartTime + FrameDuration));
+	}
+}
+
+
+uint32 FImgMediaLoader::TimeToFrame(FTimespan Time) const
+{
+	if ((Time < FTimespan::Zero()) || (Time > SequenceDuration))
+	{
+		return INDEX_NONE;
+	}
+
+	return (Time * SequenceFps).GetTicks() / ETimespan::TicksPerSecond;
+}
+
+
+void FImgMediaLoader::Update(int32 PlayHeadFrame, float PlayRate)
+{
+	// @todo ImgMedia: take PlayRate and DeltaTime into account when determining frames to load
+	
+	// determine frame numbers to be loaded
+	TArray<int32> FramesToQueue;
+	{
+		FramesToQueue.Empty(NumLoadAhead + NumLoadBehind);
+
+		int32 FrameOffset = (PlayRate >= 0.0f) ? 1 : -1;
+
+		int32 LoadAheadCount = NumLoadAhead;
+		int32 LoadAheadIndex = PlayHeadFrame;
+
+		int32 LoadBehindCount = NumLoadBehind;
+		int32 LoadBehindIndex = PlayHeadFrame - FrameOffset;
+
+		while ((LoadAheadCount > 0) || (LoadBehindCount > 0))
+		{
+			if (LoadAheadCount-- > 0)
+			{
+				if (LoadAheadIndex < 0)
+				{
+					LoadAheadIndex += ImagePaths.Num();
+				}
+				else if (LoadAheadIndex >= ImagePaths.Num())
+				{
+					LoadAheadIndex -= ImagePaths.Num();
+				}
+
+				FramesToQueue.Add(LoadAheadIndex);
+				LoadAheadIndex += FrameOffset;
+			}
+
+			if (LoadBehindCount-- > 0)
+			{
+				if (LoadBehindIndex < 0)
+				{
+					LoadBehindIndex += ImagePaths.Num();
+				}
+				else if (LoadBehindIndex >= ImagePaths.Num())
+				{
+					LoadBehindIndex -= ImagePaths.Num();
+				}
+
+				FramesToQueue.Add(LoadBehindIndex);
+				LoadBehindIndex -= FrameOffset;
+			}
+		}
+	}
+
+	// create new work items
+	TArray<FImgMediaLoaderWork*> NewWorks;
+	{
+		FScopeLock ScopeLock(&CriticalSection);
+
+		// queue up new pending frames
+		PendingFrameNumbers.Empty();
+
+		for (int32 Idx = FramesToQueue.Num() - 1; Idx >= 0; --Idx)
+		{
+			PendingFrameNumbers.Push(FramesToQueue[Idx]);
+		}
+
+		// try to abandon prior queued work if obsolete
+		for (int32 Index = QueuedFrameNumbers.Num() - 1; Index >= 0; --Index)
+		{
+			int32 FrameNumber = QueuedFrameNumbers[Index];
+
+			if (!PendingFrameNumbers.Contains(FrameNumber))
+			{
+				FImgMediaLoaderWork*& ObsoleteWork = QueuedWorks.FindChecked(FrameNumber);
+
+				if (GThreadPool->RetractQueuedWork(ObsoleteWork))
+				{
+					QueuedWorks.Remove(FrameNumber);
+					WorkPool.Push(ObsoleteWork);
+				}
+
+				QueuedFrameNumbers.RemoveAtSwap(Index);
+			}
+		}
+
+		// create or reuse work items
+		int32 NewWorkCount = GThreadPool->GetNumThreads() - QueuedWorks.Num();
+
+		while ((NewWorkCount > 0) && (PendingFrameNumbers.Num() > 0))
+		{
+			const int32 PendingFrameNumber = PendingFrameNumbers.Pop(false);
+
+			if ((Frames.FindAndTouch(PendingFrameNumber) == nullptr) && !QueuedWorks.Contains(PendingFrameNumber))
+			{
+				FImgMediaLoaderWork* NewWork = (WorkPool.Num() > 0) ? WorkPool.Pop() : new FImgMediaLoaderWork(*this, Reader.ToSharedRef());
+				NewWork->Initialize(PendingFrameNumber, ImagePaths[PendingFrameNumber]);
+				NewWorks.Add(NewWork);
+
+				QueuedFrameNumbers.Add(PendingFrameNumber);
+				QueuedWorks.Add(PendingFrameNumber, NewWork);
+				
+				--NewWorkCount;
+			}
+		}
+	}
+
+	// enqueue new work
+	for (auto& NewWork : NewWorks)
+	{
+		GThreadPool->AddQueuedWork(NewWork);
+	}
+}
+
+
+/* IImgMediaLoader interface
+ *****************************************************************************/
+
+void FImgMediaLoader::NotifyWorkComplete(FImgMediaLoaderWork& CompletedWork, int32 FrameNumber, const TSharedPtr<FImgMediaFrame, ESPMode::ThreadSafe>& Frame)
+{
+	FImgMediaLoaderWork* NewWork = nullptr;
+
+	// process completed work item
+	{
+		FScopeLock Lock(&CriticalSection);
+
+		check(QueuedWorks.FindChecked(FrameNumber) == &CompletedWork);
+		QueuedWorks.Remove(FrameNumber);
+
+		// if frame is still needed, add it to the cache
+		if (QueuedFrameNumbers.Remove(FrameNumber) > 0)
+		{
+			if (Frame.IsValid())
+			{
+				Frames.Add(FrameNumber, Frame);
+			}
+		}
+
+		// reuse work for next frame, or return it to pool
+		while (PendingFrameNumbers.Num() > 0)
+		{
+			int32 PendingFrameNumber = PendingFrameNumbers.Pop(false);
+
+			if ((Frames.FindAndTouch(PendingFrameNumber) == nullptr) && !QueuedWorks.Contains(PendingFrameNumber))
+			{
+				NewWork = &CompletedWork;
+				NewWork->Initialize(PendingFrameNumber, ImagePaths[PendingFrameNumber]);
+				
+				QueuedFrameNumbers.Add(PendingFrameNumber);
+				QueuedWorks.Add(PendingFrameNumber, NewWork);
+
+				break;
+			}
+		}
+
+		if (NewWork == nullptr)
+		{
+			WorkPool.Push(&CompletedWork);
+		}
+	}
+
+	// queue new work
+	if (NewWork != nullptr)
+	{
+		GThreadPool->AddQueuedWork(NewWork);
+	}
+}
diff --git a/Engine/Plugins/Media/ImgMedia/Source/ImgMedia/Private/Loader/ImgMediaLoaderWork.cpp b/Engine/Plugins/Media/ImgMedia/Source/ImgMedia/Private/Loader/ImgMediaLoaderWork.cpp
new file mode 100644
--- /dev/null
+++ b/Engine/Plugins/Media/ImgMedia/Source/ImgMedia/Private/Loader/ImgMediaLoaderWork.cpp
@@ -0,0 +1,93 @@
+// Copyright 1998-2016 Epic Games, Inc. All Rights Reserved.
+
+#include "ImgMediaLoaderWork.h"
+#include "ImgMediaPrivate.h"
+
+#include "Misc/ScopeLock.h"
+
+#include "IImgMediaLoader.h"
+#include "IImgMediaReader.h"
+
+
+/* FImgMediaLoaderWork structors
+ *****************************************************************************/
+
+FImgMediaLoaderWork::FImgMediaLoaderWork(IImgMediaLoader& InOwner, const TSharedRef<IImgMediaReader, ESPMode::ThreadSafe>& InReader)
+	: AutoDelete(false)
+	, Done(false)
+	, FrameNumber(INDEX_NONE)
+	, Owner(InOwner)
+	, Reader(InReader)
+{ }
+
+
+/* FImgMediaLoaderWork interface
+ *****************************************************************************/
+
+void FImgMediaLoaderWork::Initialize(int32 InFrameNumber, const FString& InImagePath)
+{
+	check(!AutoDelete);
+
+	FrameNumber = InFrameNumber;
+	ImagePath = InImagePath;
+	Done = false;
+}
+
+
+void FImgMediaLoaderWork::DeleteWhenDone()
+{
+	FScopeLock Lock(&CriticalSection);
+
+	if (Done)
+	{
+		delete this;
+	}
+	else
+	{
+		AutoDelete = true;
+	}
+}
+
+
+/* IQueuedWork interface
+ *****************************************************************************/
+
+void FImgMediaLoaderWork::Abandon()
+{
+	// not supported
+}
+
+
+void FImgMediaLoaderWork::DoThreadedWork()
+{
+	FImgMediaFrame* Frame;
+
+	if ((FrameNumber == INDEX_NONE) || ImagePath.IsEmpty())
+	{
+		Frame = nullptr;
+	}
+	else
+	{
+		// read the image frame
+		Frame = new FImgMediaFrame();
+
+		if (!Reader->ReadFrame(ImagePath, *Frame))
+		{
+			delete Frame;
+			Frame = nullptr;
+		}
+	}
+
+	// notify owner, or shut down
+	FScopeLock Lock(&CriticalSection);
+
+	if (AutoDelete)
+	{
+		delete this;
+	}
+	else
+	{
+		Done = true;
+		Owner.NotifyWorkComplete(*this, FrameNumber, MakeShareable(Frame));
+	}
+}
diff --git a/Engine/Plugins/Media/ImgMedia/Source/ImgMedia/Private/Loader/ImgMediaLoaderWork.h b/Engine/Plugins/Media/ImgMedia/Source/ImgMedia/Private/Loader/ImgMediaLoaderWork.h
new file mode 100644
--- /dev/null
+++ b/Engine/Plugins/Media/ImgMedia/Source/ImgMedia/Private/Loader/ImgMediaLoaderWork.h
@@ -0,0 +1,80 @@
+// Copyright 1998-2016 Epic Games, Inc. All Rights Reserved.
+
+#pragma once
+
+#include "CoreMinimal.h"
+#include "Misc/IQueuedWork.h"
+
+class FImgMediaLoader;
+class IImgMediaLoader;
+class IImgMediaReader;
+
+
+/**
+ * Loads a single image frame from disk.
+ */
+class FImgMediaLoaderWork
+	: public IQueuedWork
+{
+public:
+
+	/**
+	 * Create and initialize a new instance.
+	 *
+	 * @param InOwner The loader that created this work item.
+	 * @param InReader The image reader to use.
+	 */
+	FImgMediaLoaderWork(IImgMediaLoader& InOwner, const TSharedRef<IImgMediaReader, ESPMode::ThreadSafe>& InReader);
+
+public:
+
+	/**
+	 * Delete this work item when it is done.
+	 *
+	 * This method will mark the work item for shutdown. When the work is
+	 * completed, the work item will automatically delete itself. No other
+	 * methods may be called on this object after calling this one.
+	 *
+	 * @see Initialize
+	 */
+	void DeleteWhenDone();
+
+	/**
+	 * Initialize this work item.
+	 *
+	 * @param InFrameNumber The number of the image frame.
+	 * @param InImagePath The file path to the image frame to read.
+	 * @see Shutdown
+	 */
+	void Initialize(int32 InFrameNumber, const FString& InImagePath);
+
+public:
+
+	//~ IQueuedWork interface
+
+	virtual void Abandon() override;
+	virtual void DoThreadedWork() override;
+
+private:
+
+	/** Whether the work item should delete itself after it is complete. */
+	bool AutoDelete;
+
+	/** Whether the work was done. */
+	bool Done;
+
+	/** Critical section for synchronizing access to Owner. */
+	FCriticalSection CriticalSection;
+
+	/** The number of the image frame. */
+	int32 FrameNumber;
+
+	/** The file path to the image frame to read. */
+	FString ImagePath;
+
+	/** The loader that created this reader task. */
+	IImgMediaLoader& Owner;
+
+	/** The image sequence reader to use. */
+	TSharedPtr<IImgMediaReader, ESPMode::ThreadSafe> Reader;
+};
diff --git a/Engine/Plugins/Media/ImgMedia/Source/ImgMedia/Private/Readers/ExrImgMediaReader.cpp b/Engine/Plugins/Media/ImgMedia/Source/ImgMedia/Private/Readers/ExrImgMediaReader.cpp
new file mode 100644
--- /dev/null
+++ b/Engine/Plugins/Media/ImgMedia/Source/ImgMedia/Private/Readers/ExrImgMediaReader.cpp
@@ -0,0 +1,82 @@
+// Copyright 1998-2016 Epic Games, Inc. All Rights Reserved.
+
+#include "ExrImgMediaReader.h"
+
+#if IMGMEDIA_EXR_SUPPORTED_PLATFORM
+
+#include "Misc/Paths.h"
+#include "OpenExrWrapper.h"
+#include "UObject/Class.h"
+#include "UObject/UObjectGlobals.h"
+
+
+/* FExrImgMediaReader structors
+ *****************************************************************************/
+
+FExrImgMediaReader::FExrImgMediaReader()
+{
+	auto Settings = GetDefault<UImgMediaSettings>();
+	
+	FOpenExr::SetGlobalThreadCount(Settings->ExrDecoderThreads == 0
+		? FPlatformMisc::NumberOfCoresIncludingHyperthreads()
+		: Settings->ExrDecoderThreads);
+}
+
+
+/* FExrImgMediaReader interface
+ *****************************************************************************/
+
+bool FExrImgMediaReader::GetFrameInfo(const FString& ImagePath, FImgMediaFrameInfo& OutInfo)
+{
+	FRgbaInputFile InputFile(ImagePath, 2);
+	
+	return GetInfo(InputFile, OutInfo);
+}
+
+
+bool FExrImgMediaReader::ReadFrame(const FString& ImagePath, FImgMediaFrame& OutFrame)
+{
+	FRgbaInputFile InputFile(ImagePath, 2);
+	
+	if (!GetInfo(InputFile, OutFrame.Info))
+	{
+		return false;
+	}
+
+	const FIntPoint& Dim = OutFrame.Info.Dim;
+
+	if (Dim.GetMin() <= 0)
+	{
+		return false;
+	}
+
+	void* Buffer = FMemory::Malloc(Dim.X * Dim.Y * sizeof(uint16) * 4);
+	
+	InputFile.SetFrameBuffer(Buffer, Dim);
+	InputFile.ReadPixels(0, Dim.Y - 1);
+
+	OutFrame.Data = MakeShareable(Buffer, [](void* ObjectToDelete) { FMemory::Free(ObjectToDelete); });
+	OutFrame.Format = EMediaTextureSampleFormat::FloatRGBA;
+	OutFrame.Stride = Dim.X * sizeof(unsigned short) * 4;
+
+	return true;
+}
+
+
+/* FExrImgMediaReader implementation
+ *****************************************************************************/
+
+bool FExrImgMediaReader::GetInfo(FRgbaInputFile& InputFile, FImgMediaFrameInfo& OutInfo) const
+{
+	OutInfo.CompressionName = InputFile.GetCompressionName();
+	OutInfo.Dim = InputFile.GetDataWindow();
+	OutInfo.FormatName = TEXT("EXR");
+	OutInfo.Fps = (float)InputFile.GetFramesPerSecond(IMGMEDIA_DEFAULT_FPS);
+	OutInfo.Srgb = false;
+	OutInfo.UncompressedSize = InputFile.GetUncompressedSize();
+
+	return (OutInfo.UncompressedSize > 0) && (OutInfo.Dim.GetMin() > 0);
+}
+
+
+#endif //IMGMEDIA_EXR_SUPPORTED_PLATFORM
diff --git a/Engine/Plugins/Media/ImgMedia/Source/ImgMedia/Private/Readers/ExrImgMediaReader.h b/Engine/Plugins/Media/ImgMedia/Source/ImgMedia/Private/Readers/ExrImgMediaReader.h
new file mode 100644
--- /dev/null
+++ b/Engine/Plugins/Media/ImgMedia/Source/ImgMedia/Private/Readers/ExrImgMediaReader.h
@@ -0,0 +1,45 @@
+// Copyright 1998-2016 Epic Games, Inc. All Rights Reserved.
+
+#pragma once
+
+#include "ImgMediaPrivate.h"
+
+#if IMGMEDIA_EXR_SUPPORTED_PLATFORM
+
+#include "IImgMediaReader.h"
+
+class FRgbaInputFile;
+
+
+/**
+ * Implements a reader for EXR image sequences.
+ */
+class FExrImgMediaReader
+	: public IImgMediaReader
+{
+public:
+
+	/** Default constructor. */
+	FExrImgMediaReader();
+
+public:
+
+	//~ IImgMediaReader interface
+
+	virtual bool GetFrameInfo(const FString& ImagePath, FImgMediaFrameInfo& OutInfo) override;
+	virtual bool ReadFrame(const FString& ImagePath, FImgMediaFrame& OutFrame) override;
+
+protected:
+
+	/**
+	 * Get the frame information from the given input file.
+	 *
+	 * @param InputFile The input file containing the frame.
+	 * @param OutInfo Will contain the frame information.
+	 * @return true on success, false otherwise.
+	 */
+	bool GetInfo(FRgbaInputFile& InputFile, FImgMediaFrameInfo& OutInfo) const;
+};
+
+
+#endif //IMGMEDIA_EXR_SUPPORTED_PLATFORM
diff --git a/Engine/Plugins/Media/ImgMedia/Source/ImgMedia/Private/Readers/GenericImgMediaReader.cpp b/Engine/Plugins/Media/ImgMedia/Source/ImgMedia/Private/Readers/GenericImgMediaReader.cpp
new file mode 100644
--- /dev/null
+++ b/Engine/Plugins/Media/ImgMedia/Source/ImgMedia/Private/Readers/GenericImgMediaReader.cpp
@@ -0,0 +1,119 @@
+// Copyright 1998-2016 Epic Games, Inc. All Rights Reserved.
+
+#include "GenericImgMediaReader.h"
+#include "ImgMediaPrivate.h"
+
+#include "IImageWrapper.h"
+#include "IImageWrapperModule.h"
+#include "Misc/FileHelper.h"
+#include "Misc/Paths.h"
+
+
+/* Local helpers
+ *****************************************************************************/
+
+TSharedPtr<IImageWrapper> LoadImage(const FString& ImagePath, IImageWrapperModule& ImageWrapperModule, TArray<uint8>& OutBuffer, FImgMediaFrameInfo& OutInfo)
+{
+	// load image into buffer
+	if (!FFileHelper::LoadFileToArray(OutBuffer, *ImagePath))
+	{
+		UE_LOG(LogImgMedia, Warning, TEXT("FGenericImgMediaReader: Failed to load %s"), *ImagePath);
+		return nullptr;
+	}
+
+	// determine image format
+	EImageFormat ImageFormat;
+
+	const FString Extension = FPaths::GetExtension(ImagePath).ToLower();
+
+	if (Extension == TEXT("bmp"))
+	{
+		ImageFormat = EImageFormat::BMP;
+		OutInfo.FormatName = TEXT("BMP");
+	}
+	else if ((Extension == TEXT("jpg")) || (Extension == TEXT("jpeg")))
+	{
+		ImageFormat = EImageFormat::JPEG;
+		OutInfo.FormatName = TEXT("JPEG");
+	}
+	else if (Extension == TEXT("png"))
+	{
+		ImageFormat = EImageFormat::PNG;
+		OutInfo.FormatName = TEXT("PNG");
+	}
+	else
+	{
+		UE_LOG(LogImgMedia, Warning, TEXT("FGenericImgMediaReader: Unsupported file format in %s"), *ImagePath);
+		return nullptr;
+	}
+
+	// create image wrapper
+	TSharedPtr<IImageWrapper> ImageWrapper = ImageWrapperModule.CreateImageWrapper(ImageFormat);
+
+	if (!ImageWrapper.IsValid() || !ImageWrapper->SetCompressed(OutBuffer.GetData(), OutBuffer.Num()))
+	{
+		UE_LOG(LogImgMedia, Warning, TEXT("FGenericImgMediaReader: Failed to create image wrapper for %s"), *ImagePath);
+		return nullptr;
+	}
+
+	// get file info
+	OutInfo.CompressionName = TEXT("");
+	OutInfo.Dim.X = ImageWrapper->GetWidth();
+	OutInfo.Dim.Y = ImageWrapper->GetHeight();
+	OutInfo.Fps = IMGMEDIA_DEFAULT_FPS;
+	OutInfo.Srgb = true;
+	OutInfo.UncompressedSize = OutInfo.Dim.X * OutInfo.Dim.Y * 4;
+
+	return ImageWrapper;
+}
+
+
+/* FGenericImgMediaReader structors
+ *****************************************************************************/
+
+FGenericImgMediaReader::FGenericImgMediaReader(IImageWrapperModule& InImageWrapperModule)
+	: ImageWrapperModule(InImageWrapperModule)
+{ }
+
+
+/* FGenericImgMediaReader interface
+ *****************************************************************************/
+
+bool FGenericImgMediaReader::GetFrameInfo(const FString& ImagePath, FImgMediaFrameInfo& OutInfo)
+{
+	TArray<uint8> InputBuffer;
+	TSharedPtr<IImageWrapper> ImageWrapper = LoadImage(ImagePath, ImageWrapperModule, InputBuffer, OutInfo);
+
+	return ImageWrapper.IsValid();
+}
+
+
+bool FGenericImgMediaReader::ReadFrame(const FString& ImagePath, FImgMediaFrame& OutFrame)
+{
+	TArray<uint8> InputBuffer;
+	TSharedPtr<IImageWrapper> ImageWrapper = LoadImage(ImagePath, ImageWrapperModule, InputBuffer, OutFrame.Info);
+
+	if (!ImageWrapper.IsValid())
+	{
+		UE_LOG(LogImgMedia, Warning, TEXT("FGenericImgMediaReader: Failed to load image %s"), *ImagePath);
+		return false;
+	}
+
+	const TArray<uint8>* RawData = nullptr;
+
+	if (!ImageWrapper->GetRaw(ERGBFormat::BGRA, 8, RawData))
+	{
+		UE_LOG(LogImgMedia, Warning, TEXT("FGenericImgMediaReader: Failed to get image data for %s"), *ImagePath);
+		return false;
+	}
+
+	const int32 RawNum = RawData->Num();
+	void* Buffer = FMemory::Malloc(RawNum);
+	FMemory::Memcpy(Buffer, RawData->GetData(), RawNum);
+
+	OutFrame.Data = MakeShareable(Buffer, [](void* ObjectToDelete) { FMemory::Free(ObjectToDelete); });
+	OutFrame.Format = EMediaTextureSampleFormat::CharBGRA;
+	OutFrame.Stride = OutFrame.Info.Dim.X * 4;
+
+	return true;
+}
diff --git a/Engine/Plugins/Media/ImgMedia/Source/ImgMedia/Private/Readers/GenericImgMediaReader.h b/Engine/Plugins/Media/ImgMedia/Source/ImgMedia/Private/Readers/GenericImgMediaReader.h
new file mode 100644
--- /dev/null
+++ b/Engine/Plugins/Media/ImgMedia/Source/ImgMedia/Private/Readers/GenericImgMediaReader.h
@@ -0,0 +1,40 @@
+// Copyright 1998-2016 Epic Games, Inc. All Rights Reserved.
+
+#pragma once
+
+#include "IImgMediaReader.h"
+
+class FRgbaInputFile;
+class IImageWrapperModule;
+
+
+/**
+ * Implements a reader for various image sequence formats.
+ */
+class FGenericImgMediaReader
+	: public IImgMediaReader
+{
+public:
+
+	/**
+	 * Create and initialize a new instance.
+	 *
+	 * @param InImageWrapperModule The image wrapper module to use.
+	 */
+	FGenericImgMediaReader(IImageWrapperModule& InImageWrapperModule);
+
+	/** Virtual destructor. */
+	virtual ~FGenericImgMediaReader() { }
+
+public:
+
+	//~ IImgMediaReader interface
+
+	virtual bool GetFrameInfo(const FString& ImagePath, FImgMediaFrameInfo& OutInfo) override;
+	virtual bool ReadFrame(const FString& ImagePath, FImgMediaFrame& OutFrame) override;
+
+private:
+
+	/** The image wrapper module. */
+	IImageWrapperModule& ImageWrapperModule;
+};
diff --git a/Engine/Plugins/Media/ImgMedia/Source/ImgMedia/Public/ImgMediaSource.h b/Engine/Plugins/Media/ImgMedia/Source/ImgMedia/Public/ImgMediaSource.h
new file mode 100644
--- /dev/null
+++ b/Engine/Plugins/Media/ImgMedia/Source/ImgMedia/Public/ImgMediaSource.h
@@ -0,0 +1,107 @@
+// Copyright 1998-2016 Epic Games, Inc. All Rights Reserved.
+
+#pragma once
+
+#include "CoreTypes.h"
+#include "Classes/Engine/EngineTypes.h"
+#include "Containers/UnrealString.h"
+#include "UObject/ObjectMacros.h"
+#include "UObject/ScriptMacros.h"
+#include "BaseMediaSource.h"
+
+#include "ImgMediaSource.generated.h"
+
+
+/**
+ * Media source for EXR image sequences.
+ *
+ * Image sequence media sources point to a directory that contains a series of
+ * image files in which each image represents a single frame of the sequence.
+ * Currently, only EXR images are supported, but other formats, such as PNG
+ * and JPEG will be added in the future.
+ *
+ * The image sequence directory may contain sub-directories, which are called
+ * 'proxies'. Proxies can be used to provide alternative media for playback
+ * during development and testing of a game. One common scenario is the use
+ * of low resolution versions of image sequence media on computers that are
+ * too slow or don't have enough storage to play the original high-res media.
+ */
+UCLASS(BlueprintType, hidecategories=(Overrides, Playback))
+class IMGMEDIA_API UImgMediaSource
+	: public UBaseMediaSource
+{
+	GENERATED_BODY()
+
+public:
+
+	/** Default constructor. */
+	UImgMediaSource();
+
+public:
+
+	/** Overrides the default frame rate stored in the image files (0.0 = do not override). */
+	UPROPERTY(BlueprintReadWrite, EditAnywhere, Category=Sequence, AdvancedDisplay)
+	float FramesPerSecondOverride;
+
+	/** Name of the proxy directory to use. */
+	UPROPERTY(BlueprintReadWrite, EditAnywhere, Category=Sequence, AdvancedDisplay)
+	FString ProxyOverride;
+
+public:
+
+	/**
+	 * Get the names of available proxy directories.
+	 *
+	 * @param OutProxies Will contain the names of available proxy directories, if any.
+	 * @see GetSequencePath
+	 */
+	UFUNCTION(BlueprintCallable, Category="ImgMedia|ImgMediaSource")
+	void GetProxies(TArray<FString>& OutProxies) const;
+
+	/**
+	 * Get the path to the image sequence directory to be played.
+	 *
+	 * @return The file path.
+	 * @see SetSequencePath
+	 */
+	UFUNCTION(BlueprintCallable, Category="ImgMedia|ImgMediaSource")
+	const FString GetSequencePath() const
+	{
+		return SequencePath.Path;
+	}
+
+	/**
+	 * Set the path to the image sequence directory this source represents.
+	 *
+	 * @param Path The path to set.
+	 * @see GetSequencePath
+	 */
+	UFUNCTION(BlueprintCallable, Category="ImgMedia|ImgMediaSource")
+	void SetSequencePath(const FString& Path);
+
+public:
+
+	//~ IMediaOptions interface
+
+	virtual double GetMediaOption(const FName& Key, double DefaultValue) const override;
+	virtual FString GetMediaOption(const FName& Key, const FString& DefaultValue) const override;
+	virtual bool HasMediaOption(const FName& Key) const override;
+
+public:
+
+	//~ UMediaSource interface
+
+	virtual FString GetUrl() const override;
+	virtual bool Validate() const override;
+
+protected:
+
+	/** Get the full path to the image sequence. */
+	FString GetFullPath() const;
+
+protected:
+
+	/** The directory that contains the image sequence files. */
+	UPROPERTY(BlueprintReadOnly, EditAnywhere, Category=Sequence)
+	FDirectoryPath SequencePath;
+};
diff --git a/Engine/Plugins/Media/ImgMedia/Source/ImgMediaEditor/ImgMediaEditor.Build.cs b/Engine/Plugins/Media/ImgMedia/Source/ImgMediaEditor/ImgMediaEditor.Build.cs
new file mode 100644
--- /dev/null
+++ b/Engine/Plugins/Media/ImgMedia/Source/ImgMediaEditor/ImgMediaEditor.Build.cs
@@ -0,0 +1,35 @@
+// Copyright 1998-2016 Epic Games, Inc. All Rights Reserved.
+
+namespace UnrealBuildTool.Rules
+{
+	public class ImgMediaEditor : ModuleRules
+	{
+		public ImgMediaEditor(ReadOnlyTargetRules Target) : base(Target)
+		{
+			PrivateDependencyModuleNames.AddRange(
+				new string[] {
+					"Core",
+					"CoreUObject",
+					"DesktopWidgets",
+					"EditorStyle",
+					"ImgMedia",
+					"MediaAssets",
+					"Slate",
+					"SlateCore",
+					"UnrealEd",
+				});
+
+			PrivateIncludePathModuleNames.AddRange(
+				new string[] {
+					"AssetTools",
+				});
+
+			PrivateIncludePaths.AddRange(
+				new string[] {
+					"ImgMediaEditor/Private",
+					"ImgMediaEditor/Private/Customizations",
+					"ImgMediaEditor/Private/Factories",
+				});
+		}
+	}
+}
diff --git a/Engine/Plugins/Media/ImgMedia/Source/ImgMediaEditor/Private/Customizations/ImgMediaSourceCustomization.cpp b/Engine/Plugins/Media/ImgMedia/Source/ImgMediaEditor/Private/Customizations/ImgMediaSourceCustomization.cpp
new file mode 100644
--- /dev/null
+++ b/Engine/Plugins/Media/ImgMedia/Source/ImgMediaEditor/Private/Customizations/ImgMediaSourceCustomization.cpp
@@ -0,0 +1,179 @@
+// Copyright 1998-2017 Epic Games, Inc. All Rights Reserved.
+
+#include "ImgMediaSourceCustomization.h"
+
+#include "IMediaModule.h"
+#include "Misc/Paths.h"
+#include "Modules/ModuleManager.h"
+#include "EditorStyleSet.h"
+
+#include "Styling/CoreStyle.h"
+#include "Widgets/DeclarativeSyntaxSupport.h"
+#include "Widgets/SBoxPanel.h"
+#include "Widgets/Text/STextBlock.h"
+#include "Widgets/Images/SImage.h"
+#include "Widgets/Input/SDirectoryPicker.h"
+#include "Widgets/Input/SEditableTextBox.h"
+#include "Widgets/Input/SFilePathPicker.h"
+
+#include "DetailCategoryBuilder.h"
+#include "DetailLayoutBuilder.h"
+#include "DetailWidgetRow.h"
+#include "IDetailPropertyRow.h"
+
+
+#define LOCTEXT_NAMESPACE "FImgMediaSourceCustomization"
+
+
+/* IDetailCustomization interface
+ *****************************************************************************/
+
+void FImgMediaSourceCustomization::CustomizeDetails(IDetailLayoutBuilder& DetailBuilder)
+{
+	// customize 'File' category
+	IDetailCategoryBuilder& FileCategory = DetailBuilder.EditCategory("EXR");
+	{
+		// FilePath
+		SequencePathProperty = DetailBuilder.GetProperty("SequencePath");
+		{
+			IDetailPropertyRow& SequencePathRow = FileCategory.AddProperty(SequencePathProperty);
+
+			SequencePathRow
+				.ShowPropertyButtons(false)
+				.CustomWidget()
+				.NameContent()
+					[
+						SNew(SHorizontalBox)
+
+						+ SHorizontalBox::Slot()
+							.AutoWidth()
+							.VAlign(VAlign_Center)
+							[
+								SNew(STextBlock)
+									.Font(IDetailLayoutBuilder::GetDetailFont())
+									.Text(LOCTEXT("SequencePathPropertyName", "Sequence Path"))
+									.ToolTipText(SequencePathProperty->GetToolTipText())
+							]
+
+						+ SHorizontalBox::Slot()
+							.FillWidth(1.0f)
+							.HAlign(HAlign_Left)
+							.VAlign(VAlign_Center)
+							.Padding(4.0f, 0.0f, 0.0f, 0.0f)
+							[
+								SNew(SImage)
+									.Image(FCoreStyle::Get().GetBrush("Icons.Warning"))
+									.ToolTipText(LOCTEXT("SequencePathWarning", "The selected image sequence will not get packaged, because its path points to a directory outside the project's /Content/Movies/ directory."))
+									.Visibility(this, &FImgMediaSourceCustomization::HandleSequencePathWarningIconVisibility)
+							]
+					]
+				.ValueContent()
+					.MaxDesiredWidth(0.0f)
+					.MinDesiredWidth(125.0f)
+					[
+						SNew(SFilePathPicker)
+							.BrowseButtonImage(FEditorStyle::GetBrush("PropertyWindow.Button_Ellipsis"))
+							.BrowseButtonStyle(FEditorStyle::Get(), "HoverHintOnly")
+							.BrowseButtonToolTip(LOCTEXT("SequencePathBrowseButtonToolTip", "Choose a file from this computer"))
+							.BrowseDirectory_Lambda([this]() -> FString
+							{
+								const FString SequencePath = GetSequencePath();
+								return !SequencePath.IsEmpty() ? SequencePath : (FPaths::ProjectContentDir() / TEXT("Movies"));
+							})
+							.FilePath_Lambda([this]() -> FString
+							{
+								return GetSequencePath();
+							})
+							.FileTypeFilter_Lambda([]() -> FString
+							{
+								return TEXT("All files (*.*)|*.*|EXR files (*.exr)|*.exr");
+							})
+							.OnPathPicked(this, &FImgMediaSourceCustomization::HandleSequencePathPickerPathPicked)
+							.ToolTipText(LOCTEXT("SequencePathToolTip", "The path to an image sequence file on this computer"))
+					];
+		}
+	}
+
+	// add 'Proxies' category
+	IDetailCategoryBuilder& ProxiesCategory = DetailBuilder.EditCategory("Proxies", LOCTEXT("ProxiesCategoryName", "Proxies"));
+	{
+		// add 'Proxy Directories' row
+		FDetailWidgetRow& PreviewRow = ProxiesCategory.AddCustomRow(LOCTEXT("ProxiesRowFilterString", "Proxy Directories"));
+
+		PreviewRow.WholeRowContent()
+			[
+				SAssignNew(ProxiesTextBlock, SEditableTextBox)
+					.IsReadOnly(true)
+			];
+	}
+}
+
+
+/* FImgMediaSourceCustomization implementation
+ *****************************************************************************/
+
+FString FImgMediaSourceCustomization::GetSequencePath() const
+{
+	FString FilePath;
+	SequencePathProperty->GetChildHandle("Path")->GetValue(FilePath);
+
+	return FilePath;
+}
+
+
+/* FImgMediaSourceCustomization callbacks
+ *****************************************************************************/
+
+void FImgMediaSourceCustomization::HandleSequencePathPickerPathPicked(const FString& PickedPath)
+{
+	// fully expand path & strip optional file name
+	const FString FullPath = PickedPath.StartsWith(TEXT("./"))
+		? FPaths::ConvertRelativePathToFull(FPaths::ProjectContentDir(), PickedPath.RightChop(2))
+		: PickedPath;
+
+	const FString FullDir = FPaths::ConvertRelativePathToFull(FPaths::FileExists(FullPath) ? FPaths::GetPath(FullPath) : FullPath);
+	const FString FullGameContentDir = FPaths::ConvertRelativePathToFull(FPaths::ProjectContentDir());
+
+	FString PickedDir = FullDir;
+
+	if (PickedDir.StartsWith(FullGameContentDir))
+	{
+		FPaths::MakePathRelativeTo(PickedDir, *FullGameContentDir);
+		PickedDir = FString(TEXT("./")) + PickedDir;
+	}
+
+	// update property
+	TSharedPtr<IPropertyHandle> SequencePathPathProperty = SequencePathProperty->GetChildHandle("Path");
+	SequencePathPathProperty->SetValue(PickedDir);
+}
+
+
+EVisibility FImgMediaSourceCustomization::HandleSequencePathWarningIconVisibility() const
+{
+	FString FilePath;
+
+	if ((SequencePathProperty->GetValue(FilePath) != FPropertyAccess::Success) || FilePath.IsEmpty() || FilePath.Contains(TEXT("://")))
+	{
+		return EVisibility::Hidden;
+	}
+
+	const FString FullMoviesPath = FPaths::ConvertRelativePathToFull(FPaths::ProjectContentDir() / TEXT("Movies"));
+	const FString FullPath = FPaths::ConvertRelativePathToFull(FPaths::IsRelative(FilePath) ? FPaths::ProjectContentDir() / FilePath : FilePath);
+
+	if (FullPath.StartsWith(FullMoviesPath))
+	{
+		if (FPaths::FileExists(FullPath))
+		{
+			return EVisibility::Hidden;
+		}
+
+		// file doesn't exist
+		return EVisibility::Visible;
+	}
+
+	// file not inside Movies folder
+	return EVisibility::Visible;
+}
+
+
+#undef LOCTEXT_NAMESPACE
diff --git a/Engine/Plugins/Media/ImgMedia/Source/ImgMediaEditor/Private/Customizations/ImgMediaSourceCustomization.h b/Engine/Plugins/Media/ImgMedia/Source/ImgMediaEditor/Private/Customizations/ImgMediaSourceCustomization.h
new file mode 100644
--- /dev/null
+++ b/Engine/Plugins/Media/ImgMedia/Source/ImgMediaEditor/Private/Customizations/ImgMediaSourceCustomization.h
@@ -0,0 +1,69 @@
+// Copyright 1998-2017 Epic Games, Inc. All Rights Reserved.
+
+#pragma once
+
+#include "CoreTypes.h"
+#include "Containers/UnrealString.h"
+#include "Layout/Visibility.h"
+#include "IDetailCustomization.h"
+#include "Templates/SharedPointer.h"
+
+class IDetailLayoutBuilder;
+class IPropertyHandle;
+class SEditableTextBox;
+
+
+/**
+ * Implements a details view customization for the UImgMediaSource class.
+ */
+class FImgMediaSourceCustomization
+	: public IDetailCustomization
+{
+public:
+
+	/** Virtual destructor. */
+	~FImgMediaSourceCustomization() { }
+
+public:
+
+	//~ IDetailCustomization interface
+
+	virtual void CustomizeDetails(IDetailLayoutBuilder& DetailBuilder) override;
+
+public:
+
+	/**
+	 * Creates an instance of this class.
+	 *
+	 * @return The new instance.
+	 */
+	static TSharedRef<IDetailCustomization> MakeInstance()
+	{
+		return MakeShareable(new FImgMediaSourceCustomization());
+	}
+
+protected:
+
+	/**
+	 * Get the path to the currently selected image sequence.
+	 *
+	 * @return Sequence path string.
+	 */
+	FString GetSequencePath() const;
+
+private:
+
+	/** Callback for picking a path in the source directory picker. */
+	void HandleSequencePathPickerPathPicked(const FString& PickedPath);
+
+	/** Callback for getting the visibility of warning icon for invalid SequencePath paths. */
+	EVisibility HandleSequencePathWarningIconVisibility() const;
+
+private:
+
+	/** Text block widget showing the found proxy directories. */
+	TSharedPtr<SEditableTextBox> ProxiesTextBlock;
+
+	/** Pointer to the SequencePath.Path property handle. */
+	TSharedPtr<IPropertyHandle> SequencePathProperty;
+};
diff --git a/Engine/Plugins/Media/ImgMedia/Source/ImgMediaEditor/Private/ImgMediaEditorModule.cpp b/Engine/Plugins/Media/ImgMedia/Source/ImgMediaEditor/Private/ImgMediaEditorModule.cpp
new file mode 100644
--- /dev/null
+++ b/Engine/Plugins/Media/ImgMedia/Source/ImgMediaEditor/Private/ImgMediaEditorModule.cpp
@@ -0,0 +1,62 @@
+// Copyright 1998-2016 Epic Games, Inc. All Rights Reserved.
+
+#include "Modules/ModuleInterface.h"
+#include "Modules/ModuleManager.h"
+#include "PropertyEditorModule.h"
+#include "UObject/NameTypes.h"
+
+#include "ImgMediaSource.h"
+
+#include "Customizations/ImgMediaSourceCustomization.h"
+
+
+/**
+ * Implements the ImgMediaEditor module.
+ */
+class FImgMediaEditorModule
+	: public IModuleInterface
+{
+public:
+
+	//~ IModuleInterface interface
+
+	virtual void StartupModule() override
+	{
+		RegisterCustomizations();
+	}
+
+	virtual void ShutdownModule() override
+	{
+		UnregisterCustomizations();
+	}
+
+protected:
+
+	/** Register details view customizations. */
+	void RegisterCustomizations()
+	{
+		ImgMediaSourceName = UImgMediaSource::StaticClass()->GetFName();
+
+		FPropertyEditorModule& PropertyModule = FModuleManager::LoadModuleChecked<FPropertyEditorModule>("PropertyEditor");
+		{
+			PropertyModule.RegisterCustomClassLayout(ImgMediaSourceName, FOnGetDetailCustomizationInstance::CreateStatic(&FImgMediaSourceCustomization::MakeInstance));
+		}
+	}
+
+	/** Unregister details view customizations. */
+	void UnregisterCustomizations()
+	{
+		FPropertyEditorModule& PropertyModule = FModuleManager::LoadModuleChecked<FPropertyEditorModule>("PropertyEditor");
+		{
+			PropertyModule.UnregisterCustomClassLayout(ImgMediaSourceName);
+		}
+	}
+
+private:
+
+	/** Class names. */
+	FName ImgMediaSourceName;
+};
+
+
+IMPLEMENT_MODULE(FImgMediaEditorModule, ImgMediaEditor);
diff --git a/Engine/Plugins/Media/ImgMedia/Source/ImgMediaFactory/ImgMediaFactory.Build.cs b/Engine/Plugins/Media/ImgMedia/Source/ImgMediaFactory/ImgMediaFactory.Build.cs
new file mode 100644
--- /dev/null
+++ b/Engine/Plugins/Media/ImgMedia/Source/ImgMediaFactory/ImgMediaFactory.Build.cs
@@ -0,0 +1,49 @@
+// Copyright 1998-2016 Epic Games, Inc. All Rights Reserved.
+
+namespace UnrealBuildTool.Rules
+{
+	public class ImgMediaFactory : ModuleRules
+	{
+		public ImgMediaFactory(ReadOnlyTargetRules Target) : base(Target)
+		{
+			DynamicallyLoadedModuleNames.AddRange(
+				new string[] {
+					"Media",
+				});
+
+			PrivateDependencyModuleNames.AddRange(
+				new string[] {
+					"MediaAssets",
+				});
+
+			PrivateIncludePathModuleNames.AddRange(
+				new string[] {
+					"Media",
+					"ImgMedia",
+				});
+
+			PrivateIncludePaths.AddRange(
+				new string[] {
+					"ImgMediaFactory/Private",
+				});
+
+			PublicDependencyModuleNames.AddRange(
+				new string[] {
+					"Core",
+					"CoreUObject",
+				});
+
+			if (Target.Type == TargetRules.TargetType.Editor)
+			{
+				DynamicallyLoadedModuleNames.Add("Settings");
+				PrivateIncludePathModuleNames.Add("Settings");
+			}
+
+			if ((Target.Platform == UnrealTargetPlatform.Win32) ||
+				(Target.Platform == UnrealTargetPlatform.Win64))
+			{
+				DynamicallyLoadedModuleNames.Add("ImgMedia");
+			}
+		}
+	}
+}
diff --git a/Engine/Plugins/Media/ImgMedia/Source/OpenExrWrapper/OpenExrWrapper.Build.cs b/Engine/Plugins/Media/ImgMedia/Source/OpenExrWrapper/OpenExrWrapper.Build.cs
new file mode 100644
--- /dev/null
+++ b/Engine/Plugins/Media/ImgMedia/Source/OpenExrWrapper/OpenExrWrapper.Build.cs
@@ -0,0 +1,31 @@
+// Copyright 1998-2016 Epic Games, Inc. All Rights Reserved.
+
+namespace UnrealBuildTool.Rules
+{
+	public class OpenExrWrapper : ModuleRules
+	{
+		public OpenExrWrapper(ReadOnlyTargetRules Target) : base(Target)
+		{
+			PCHUsage = ModuleRules.PCHUsageMode.NoSharedPCHs;
+			PrivatePCHHeaderFile = "Public/OpenExrWrapper.h";
+
+            bUseRTTI = true;
+
+            PublicDependencyModuleNames.Add("Core");
+
+            if ((Target.Platform == UnrealTargetPlatform.Win64) ||
+                (Target.Platform == UnrealTargetPlatform.Win32) ||
+                (Target.Platform == UnrealTargetPlatform.Mac))
+            {
+                AddEngineThirdPartyPrivateStaticDependencies(Target, "UEOpenExr");
+                AddEngineThirdPartyPrivateStaticDependencies(Target, "zlib");
+            }
+            else
+            {
+                System.Console.WriteLine("OpenExrWrapper does not supported this platform");
+            }
+
+            PrivateIncludePaths.Add("OpenExrWrapper/Private");
+        }
+	}
+}
diff --git a/Engine/Plugins/Media/ImgMedia/Source/OpenExrWrapper/Public/OpenExrWrapper.h b/Engine/Plugins/Media/ImgMedia/Source/OpenExrWrapper/Public/OpenExrWrapper.h
new file mode 100644
--- /dev/null
+++ b/Engine/Plugins/Media/ImgMedia/Source/OpenExrWrapper/Public/OpenExrWrapper.h
@@ -0,0 +1,40 @@
+// Copyright 1998-2016 Epic Games, Inc. All Rights Reserved.
+
+#pragma once
+
+#include "CoreTypes.h"
+#include "Math/IntPoint.h"
+
+class FString;
+
+
+class OPENEXRWRAPPER_API FOpenExr
+{
+public:
+
+	static void SetGlobalThreadCount(uint16 ThreadCount);
+};
+
+
+class OPENEXRWRAPPER_API FRgbaInputFile
+{
+public:
+
+	FRgbaInputFile(const FString& FilePath);
+	FRgbaInputFile(const FString& FilePath, uint16 ThreadCount);
+	~FRgbaInputFile();
+
+public:
+
+	const TCHAR* GetCompressionName() const;
+	FIntPoint GetDataWindow() const;
+	double GetFramesPerSecond(double DefaultValue) const;
+	int32 GetUncompressedSize() const;
+	bool IsComplete() const;
+	void ReadPixels(int32 StartY, int32 EndY);
+	void SetFrameBuffer(void* Buffer, const FIntPoint& Stride);
+
+private:
+
+	void* InputFile;
+};
diff --git a/Engine/Plugins/Media/MediaCompositing/Source/MediaCompositing/MediaCompositing.Build.cs b/Engine/Plugins/Media/MediaCompositing/Source/MediaCompositing/MediaCompositing.Build.cs
new file mode 100644
--- /dev/null
+++ b/Engine/Plugins/Media/MediaCompositing/Source/MediaCompositing/MediaCompositing.Build.cs
@@ -0,0 +1,39 @@
+// Copyright 1998-2016 Epic Games, Inc. All Rights Reserved.
+
+namespace UnrealBuildTool.Rules
+{
+	public class MediaCompositing : ModuleRules
+	{
+		public MediaCompositing(ReadOnlyTargetRules Target) : base(Target)
+		{
+			PublicDependencyModuleNames.AddRange(
+				new string[] {
+					"Core",
+					"CoreUObject",
+					"Engine",
+					"MovieScene",
+					"MovieSceneTracks",
+				});
+
+			PrivateDependencyModuleNames.AddRange(
+				new string[] {
+					"Core",
+					"CoreUObject",
+					"Engine",
+					"Media",
+					"MediaAssets",
+					"RenderCore",
+					"RHI",
+					"ShaderCore",
+				});
+
+			PrivateIncludePaths.AddRange(
+				new string[] {
+					"MediaCompositing/Private",
+					"MediaCompositing/Private/Actors",
+					"MediaCompositing/Private/Components",
+					"MediaCompositing/Private/MovieScene",
+				});
+		}
+	}
+}
diff --git a/Engine/Plugins/Media/MediaCompositing/Source/MediaCompositing/Private/Actors/MediaPlane.cpp b/Engine/Plugins/Media/MediaCompositing/Source/MediaCompositing/Private/Actors/MediaPlane.cpp
new file mode 100644
--- /dev/null
+++ b/Engine/Plugins/Media/MediaCompositing/Source/MediaCompositing/Private/Actors/MediaPlane.cpp
@@ -0,0 +1,15 @@
+// Copyright 1998-2016 Epic Games, Inc. All Rights Reserved.
+
+#include "MediaPlane.h"
+#include "MediaPlaneComponent.h"
+
+
+/* AMediaPlane structors
+ *****************************************************************************/
+
+AMediaPlane::AMediaPlane(const FObjectInitializer& ObjectInitializer)
+	: Super(ObjectInitializer)
+{
+	MediaPlane = ObjectInitializer.CreateDefaultSubobject<UMediaPlaneComponent>(this, "MediaPlaneComponent");
+	RootComponent = MediaPlane;
+}
\ No newline at end of file
diff --git a/Engine/Plugins/Media/MediaCompositing/Source/MediaCompositing/Private/Components/MediaPlaneComponent.cpp b/Engine/Plugins/Media/MediaCompositing/Source/MediaCompositing/Private/Components/MediaPlaneComponent.cpp
new file mode 100644
--- /dev/null
+++ b/Engine/Plugins/Media/MediaCompositing/Source/MediaCompositing/Private/Components/MediaPlaneComponent.cpp
@@ -0,0 +1,579 @@
+// Copyright 1998-2016 Epic Games, Inc. All Rights Reserved.
+
+#include "MediaPlaneComponent.h"
+
+#include "EngineGlobals.h"
+#include "RHI.h"
+#include "RenderResource.h"
+#include "RenderingThread.h"
+#include "VertexFactory.h"
+#include "PackedNormal.h"
+#include "LocalVertexFactory.h"
+#include "PrimitiveViewRelevance.h"
+#include "Materials/MaterialInterface.h"
+#include "Materials/MaterialInstanceDynamic.h"
+#include "PrimitiveSceneProxy.h"
+#include "Engine/CollisionProfile.h"
+#include "Curves/CurveFloat.h"
+#include "SceneManagement.h"
+#include "Engine/Engine.h"
+#include "Engine/LevelStreaming.h"
+#include "LevelUtils.h"
+#include "ShowFlags.h"
+#include "Camera/CameraTypes.h"
+#include "Materials/MaterialInstanceDynamic.h"
+#include "Engine/SceneCapture2D.h"
+#include "Engine/Texture2DDynamic.h"
+#include "Engine/TextureRenderTarget2D.h"
+#include "DynamicMeshBuilder.h"
+#include "Components/SceneCaptureComponent2D.h"
+#include "MediaPlaneFrustumComponent.h"
+#include "MediaPlaneComponent.h"
+
+namespace
+{
+	FMatrix CalculateProjectionMatrix(const FMinimalViewInfo& MinimalView)
+	{
+		FMatrix ProjectionMatrix;
+
+		if (MinimalView.ProjectionMode == ECameraProjectionMode::Orthographic)
+		{
+			const float YScale = 1.0f / MinimalView.AspectRatio;
+
+			const float HalfOrthoWidth = MinimalView.OrthoWidth / 2.0f;
+			const float ScaledOrthoHeight = MinimalView.OrthoWidth / 2.0f * YScale;
+
+			const float NearPlane = MinimalView.OrthoNearClipPlane;
+			const float FarPlane = MinimalView.OrthoFarClipPlane;
+
+			const float ZScale = 1.0f / (FarPlane - NearPlane);
+			const float ZOffset = -NearPlane;
+
+			ProjectionMatrix = FReversedZOrthoMatrix(
+				HalfOrthoWidth,
+				ScaledOrthoHeight,
+				ZScale,
+				ZOffset
+				);
+		}
+		else
+		{
+			// Avoid divide by zero in the projection matrix calculation by clamping FOV
+			ProjectionMatrix = FReversedZPerspectiveMatrix(
+				FMath::Max(0.001f, MinimalView.FOV) * (float)PI / 360.0f,
+				MinimalView.AspectRatio,
+				1.0f,
+				GNearClippingPlane );
+		}
+
+		if (!MinimalView.OffCenterProjectionOffset.IsZero())
+		{
+			const float Left = -1.0f + MinimalView.OffCenterProjectionOffset.X;
+			const float Right = Left + 2.0f;
+			const float Bottom = -1.0f + MinimalView.OffCenterProjectionOffset.Y;
+			const float Top = Bottom + 2.0f;
+			ProjectionMatrix.M[2][0] = (Left + Right) / (Left - Right);
+			ProjectionMatrix.M[2][1] = (Bottom + Top) / (Bottom - Top);
+		}
+
+		return ProjectionMatrix;
+	}
+
+	class FMediaPlaneVertexBuffer : public FVertexBuffer
+	{
+	public:
+
+		TArray<FDynamicMeshVertex> Vertices;
+
+		virtual void InitRHI() override
+		{
+			FRHIResourceCreateInfo CreateInfo;
+			void* VertexBufferData = nullptr;
+			VertexBufferRHI = RHICreateAndLockVertexBuffer(Vertices.Num() * sizeof(FDynamicMeshVertex),BUF_Static,CreateInfo, VertexBufferData);
+
+			// Copy the vertex data into the vertex buffer.
+			FMemory::Memcpy(VertexBufferData,Vertices.GetData(),Vertices.Num() * sizeof(FDynamicMeshVertex));
+			RHIUnlockVertexBuffer(VertexBufferRHI);
+		}
+	};
+
+	struct FMediaPlaneVertexFactory : public FLocalVertexFactory
+	{
+		/** Initialization */
+		void Init(const FMediaPlaneVertexBuffer* VertexBuffer)
+		{
+			check(IsInRenderingThread())
+
+			FDataType NewData;
+			NewData.PositionComponent = 			STRUCTMEMBER_VERTEXSTREAMCOMPONENT(VertexBuffer,FDynamicMeshVertex,Position,VET_Float3);
+			NewData.TangentBasisComponents[0] = 	STRUCTMEMBER_VERTEXSTREAMCOMPONENT(VertexBuffer,FDynamicMeshVertex,TangentX,VET_PackedNormal);
+			NewData.TangentBasisComponents[1] = 	STRUCTMEMBER_VERTEXSTREAMCOMPONENT(VertexBuffer,FDynamicMeshVertex,TangentZ,VET_PackedNormal);
+			NewData.ColorComponent = 				STRUCTMEMBER_VERTEXSTREAMCOMPONENT(VertexBuffer,FDynamicMeshVertex,Color,VET_Color);
+
+			NewData.TextureCoordinates.Add(
+				FVertexStreamComponent(VertexBuffer,STRUCT_OFFSET(FDynamicMeshVertex,TextureCoordinate),sizeof(FDynamicMeshVertex),VET_Float2)
+				);
+
+			SetData(NewData);
+		}
+	};
+
+	class FMediaPlaneIndexBuffer : public FIndexBuffer
+	{
+	public:
+		TArray<uint16> Indices;
+
+		virtual void InitRHI() override
+		{
+			FRHIResourceCreateInfo CreateInfo;
+			void* Buffer = nullptr;
+			IndexBufferRHI = RHICreateAndLockIndexBuffer(sizeof(uint16), Indices.Num() * sizeof(uint16), BUF_Static, CreateInfo, Buffer);
+
+			// Copy the index data into the index buffer.		
+			FMemory::Memcpy(Buffer, Indices.GetData(), Indices.Num() * sizeof(uint16));
+			RHIUnlockIndexBuffer(IndexBufferRHI);
+		}
+	};
+
+	/** Represents a sprite to the scene manager. */
+	class FMediaPlaneSceneProxy : public FPrimitiveSceneProxy
+	{
+	public:
+
+		/** Initialization constructor. */
+		FMediaPlaneSceneProxy(UMediaPlaneComponent* InComponent)
+			: FPrimitiveSceneProxy(InComponent)
+		{
+			AActor* Owner = InComponent->GetOwner();
+			if (Owner)
+			{
+				// Level colorization
+				ULevel* Level = Owner->GetLevel();
+				ULevelStreaming* LevelStreaming = FLevelUtils::FindStreamingLevel( Level );
+				if ( LevelStreaming )
+				{
+					// Selection takes priority over level coloration.
+					LevelColor = LevelStreaming->LevelColor;
+				}
+			}
+
+			Material = InComponent->GetPlane().DynamicMaterial ? InComponent->GetPlane().DynamicMaterial : InComponent->GetPlane().Material;
+			if (Material)
+			{
+				MaterialRelevance |= Material->GetRelevance(GetScene().GetFeatureLevel());
+			}
+
+			FColor NewPropertyColor;
+			GEngine->GetPropertyColorationColor(InComponent, NewPropertyColor);
+			PropertyColor = NewPropertyColor;
+		}
+
+		~FMediaPlaneSceneProxy()
+		{
+			VertexBuffer.ReleaseResource();
+			IndexBuffer.ReleaseResource();
+			VertexFactory.ReleaseResource();
+		}
+
+		virtual void CreateRenderThreadResources() override
+		{
+			BuildMesh();
+
+			VertexFactory.Init(&VertexBuffer);
+
+			VertexBuffer.InitResource();
+			IndexBuffer.InitResource();
+			VertexFactory.InitResource();
+		}
+
+		void BuildMesh()
+		{
+			VertexBuffer.Vertices.Empty(4);
+			VertexBuffer.Vertices.AddUninitialized(4);
+
+			// Set up the sprite vertex positions and texture coordinates.
+			VertexBuffer.Vertices[0].Position  = FVector(0, -1.f,  1.f);
+			VertexBuffer.Vertices[1].Position  = FVector(0, -1.f, -1.f);
+			VertexBuffer.Vertices[2].Position  = FVector(0,  1.f,  1.f);
+			VertexBuffer.Vertices[3].Position  = FVector(0,  1.f, -1.f);
+
+			VertexBuffer.Vertices[0].TextureCoordinate = FVector2D(0,0);
+			VertexBuffer.Vertices[1].TextureCoordinate = FVector2D(0,1);
+			VertexBuffer.Vertices[2].TextureCoordinate = FVector2D(1,0);
+			VertexBuffer.Vertices[3].TextureCoordinate = FVector2D(1,1);
+
+			IndexBuffer.Indices.Empty(6);
+			IndexBuffer.Indices.AddUninitialized(6);
+
+			IndexBuffer.Indices[0] = 0;
+			IndexBuffer.Indices[1] = 1;
+			IndexBuffer.Indices[2] = 2;
+			IndexBuffer.Indices[3] = 1;
+			IndexBuffer.Indices[4] = 3;
+			IndexBuffer.Indices[5] = 2;
+		}
+
+		virtual void DrawStaticElements(FStaticPrimitiveDrawInterface* PDI) override
+		{
+			QUICK_SCOPE_CYCLE_COUNTER( STAT_MediaPlaneSceneProxy_DrawStaticElements );
+
+			if (Material)
+			{
+				FMeshBatch Mesh;
+				Mesh.VertexFactory           = &VertexFactory;
+				Mesh.MaterialRenderProxy     = Material->GetRenderProxy(false);
+				Mesh.ReverseCulling          = IsLocalToWorldDeterminantNegative();
+				Mesh.CastShadow              = false;
+				Mesh.DepthPriorityGroup      = SDPG_World;
+				Mesh.Type                    = PT_TriangleList;
+				Mesh.bDisableBackfaceCulling = true;
+
+				FMeshBatchElement& BatchElement = Mesh.Elements[0];
+				BatchElement.IndexBuffer = &IndexBuffer;
+				BatchElement.FirstIndex             = 0;
+				BatchElement.MinVertexIndex         = 0;
+				BatchElement.MaxVertexIndex         = 3;
+				BatchElement.NumPrimitives          = 2;
+				BatchElement.PrimitiveUniformBufferResource = &GetUniformBuffer();
+
+				PDI->DrawMesh(Mesh, 1.f);
+			}
+		}
+
+		virtual void GetDynamicMeshElements(const TArray<const FSceneView*>& Views, const FSceneViewFamily& ViewFamily, uint32 VisibilityMap, FMeshElementCollector& Collector) const
+		{
+			QUICK_SCOPE_CYCLE_COUNTER( STAT_MediaPlaneSceneProxy_GetDynamicMeshElements );
+
+			if (!Material)
+			{
+				return;
+			}
+
+			for (int32 ViewIndex = 0; ViewIndex < Views.Num(); ViewIndex++)
+			{
+				if (!(VisibilityMap & (1 << ViewIndex)))
+				{
+					continue;
+				}
+
+				const FSceneView* View = Views[ViewIndex];
+
+				// Draw the mesh.
+				FMeshBatch& Mesh = Collector.AllocateMesh();
+				Mesh.VertexFactory = &VertexFactory;
+				Mesh.ReverseCulling = IsLocalToWorldDeterminantNegative();
+				Mesh.CastShadow = false;
+				Mesh.bDisableBackfaceCulling = false;
+				Mesh.Type = PT_TriangleList;
+				Mesh.DepthPriorityGroup = (ESceneDepthPriorityGroup)GetDepthPriorityGroup(View);
+				Mesh.bCanApplyViewModeOverrides = true;
+				Mesh.MaterialRenderProxy = Material->GetRenderProxy(View->Family->EngineShowFlags.Selection && IsSelected(), IsHovered());
+
+				FMeshBatchElement& BatchElement = Mesh.Elements[0];
+				BatchElement.IndexBuffer = &IndexBuffer;
+				BatchElement.FirstIndex			= 0;
+				BatchElement.MinVertexIndex		= 0;
+				BatchElement.MaxVertexIndex		= 3;
+				BatchElement.NumPrimitives		= 2;
+				BatchElement.PrimitiveUniformBufferResource = &GetUniformBuffer();
+
+				Collector.AddMesh(ViewIndex, Mesh);
+
+#if !(UE_BUILD_SHIPPING || UE_BUILD_TEST)
+				RenderBounds(Collector.GetPDI(ViewIndex), View->Family->EngineShowFlags, GetBounds(), IsSelected());
+#endif
+			}
+		}
+
+		virtual FPrimitiveViewRelevance GetViewRelevance(const FSceneView* View) const override
+		{
+			FPrimitiveViewRelevance Result;
+			Result.bDrawRelevance = IsShown(View) && View->Family->EngineShowFlags.MediaPlanes;
+			Result.bRenderCustomDepth = ShouldRenderCustomDepth();
+			Result.bRenderInMainPass = ShouldRenderInMainPass();
+			Result.bUsesLightingChannels = GetLightingChannelMask() != GetDefaultLightingChannelMask();
+
+			Result.bShadowRelevance = IsShadowCast(View);
+			if( IsRichView(*View->Family) ||
+				View->Family->EngineShowFlags.Bounds ||
+				View->Family->EngineShowFlags.Collision ||
+				IsSelected() ||
+				IsHovered())
+			{
+				Result.bDynamicRelevance = true;
+			}
+			else
+			{
+				Result.bStaticRelevance = true;
+			}
+			
+			MaterialRelevance.SetPrimitiveViewRelevance(Result);
+			return Result;
+		}
+		virtual bool CanBeOccluded() const override { return !MaterialRelevance.bDisableDepthTest; }
+		virtual uint32 GetMemoryFootprint() const override { return sizeof(*this) + GetAllocatedSize(); }
+		uint32 GetAllocatedSize() const { return FPrimitiveSceneProxy::GetAllocatedSize(); }
+
+	private:
+		UMaterialInterface* Material;
+		FMaterialRelevance MaterialRelevance;
+		FMediaPlaneVertexBuffer VertexBuffer;
+		FMediaPlaneIndexBuffer IndexBuffer;
+		FMediaPlaneVertexFactory VertexFactory;
+	};
+}
+
+FMediaPlaneParameters::FMediaPlaneParameters()
+	: Material(LoadObject<UMaterialInterface>(nullptr, TEXT("/MediaCompositing/DefaultMediaPlaneMaterial.DefaultMediaPlaneMaterial")))
+	, TextureParameterName("InputTexture")
+	, bFillScreen(true)
+	, FillScreenAmount(100, 100)
+	, FixedSize(100,100)
+	, RenderTexture(nullptr)
+	, DynamicMaterial(nullptr)
+{
+}
+
+UMediaPlaneComponent::UMediaPlaneComponent(const FObjectInitializer& Init)
+	: Super(Init)
+{
+	SetCollisionProfileName(UCollisionProfile::NoCollision_ProfileName);
+	bUseAsOccluder = false;
+	bTickInEditor = true;
+	PrimaryComponentTick.bCanEverTick = true;
+	bReenetrantTranformChange = false;
+}
+
+void UMediaPlaneComponent::OnRegister()
+{
+	Super::OnRegister();
+	UpdateMaterialParametersForMedia();
+	UpdateTransformScale();
+
+#if WITH_EDITORONLY_DATA
+	if (AActor* ComponentOwner = GetOwner())
+	{
+		if (!EditorFrustum)
+		{
+			EditorFrustum = NewObject<UMediaPlaneFrustumComponent>(ComponentOwner, NAME_None, RF_Transactional | RF_TextExportTransient);
+			EditorFrustum->SetupAttachment(this);
+			/*EditorFrustum->bIsEditorOnly = true;*/
+			EditorFrustum->CreationMethod = CreationMethod;
+			EditorFrustum->RegisterComponentWithWorld(GetWorld());
+		}
+	}
+#endif
+}
+
+void UMediaPlaneComponent::TickComponent(float DeltaTime, ELevelTick TickType, FActorComponentTickFunction* ThisTickFunction)
+{
+	Super::TickComponent(DeltaTime, TickType, ThisTickFunction);
+	UpdateTransformScale();
+}
+
+void UMediaPlaneComponent::OnUpdateTransform(EUpdateTransformFlags UpdateTransformFlags, ETeleportType Teleport)
+{
+	Super::OnUpdateTransform(UpdateTransformFlags, Teleport);
+
+	if (!bReenetrantTranformChange)
+	{
+		bReenetrantTranformChange = true;
+		UpdateTransformScale();
+		bReenetrantTranformChange = false;
+	}
+}
+
+void UMediaPlaneComponent::SetMediaPlane(FMediaPlaneParameters NewPlane)
+{
+	Plane = NewPlane;
+
+	UpdateMaterialParametersForMedia();
+}
+
+void UMediaPlaneComponent::OnRenderTextureChanged()
+{
+	UpdateMaterialParametersForMedia();
+}
+
+void UMediaPlaneComponent::UpdateTransformScale()
+{
+	// Cache the view projection matrices of our target
+	AActor* ViewTarget = FindViewTarget();
+	if (ViewTarget && Plane.bFillScreen)
+	{
+		GetProjectionMatricesFromViewTarget(ViewTarget, ViewProjectionMatrix, InvViewProjectionMatrix);
+
+		const FMatrix LocalToWorld = GetComponentTransform().ToMatrixNoScale();
+		const FMatrix WorldToLocal = LocalToWorld.Inverse();
+		const FMatrix ScreenToLocalSpace = InvViewProjectionMatrix * WorldToLocal;
+
+
+		// Just use the current view projection matrix
+		FVector4 HGLocalPosition = (LocalToWorld * ViewProjectionMatrix).TransformPosition(FVector::ZeroVector);
+		FVector ScreenSpaceLocalPosition = HGLocalPosition;
+		if (HGLocalPosition.W != 0.0f)
+		{
+			ScreenSpaceLocalPosition /= HGLocalPosition.W;
+		}
+
+		FVector HorizontalScale		= UMediaPlaneComponent::TransfromFromProjection(ScreenToLocalSpace, FVector4(Plane.FillScreenAmount.X/100.f, 0.f, ScreenSpaceLocalPosition.Z, 1.0f));
+		FVector VerticalScale		= UMediaPlaneComponent::TransfromFromProjection(ScreenToLocalSpace, FVector4(0.f, Plane.FillScreenAmount.Y/100.f, ScreenSpaceLocalPosition.Z, 1.0f));
+
+		SetRelativeScale3D(FVector(RelativeScale3D.X, HorizontalScale.Size(), VerticalScale.Size()));
+		SetRelativeLocation(FVector(RelativeLocation.X, 0.f, 0.f));
+	}
+	else
+	{
+		SetRelativeScale3D(FVector(RelativeScale3D.X, Plane.FixedSize.X*.5f, Plane.FixedSize.Y*.5f));
+	}
+}
+
+void UMediaPlaneComponent::UpdateMaterialParametersForMedia()
+{
+	if (!Plane.TextureParameterName.IsNone() && Plane.Material && Plane.RenderTexture)
+	{
+		if (!Plane.DynamicMaterial)
+		{
+			Plane.DynamicMaterial = UMaterialInstanceDynamic::Create(Plane.Material, this);
+			Plane.DynamicMaterial->SetFlags(RF_Transient);
+		}
+
+		Plane.DynamicMaterial->SetTextureParameterValue(Plane.TextureParameterName, Plane.RenderTexture);
+	}
+	else
+	{
+		Plane.DynamicMaterial = nullptr;
+	}
+
+	MarkRenderStateDirty();
+
+#if WITH_EDITORONLY_DATA
+	if (EditorFrustum)
+	{
+		EditorFrustum->MarkRenderStateDirty();
+	}
+#endif
+}
+
+FPrimitiveSceneProxy* UMediaPlaneComponent::CreateSceneProxy()
+{
+	return new FMediaPlaneSceneProxy(this);
+}
+
+UMaterialInterface* UMediaPlaneComponent::GetMaterial(int32 Index) const
+{
+	if (Index == 0)
+	{
+		return Plane.Material;
+	}
+	return nullptr;
+}
+
+void UMediaPlaneComponent::SetMaterial(int32 ElementIndex, UMaterialInterface* NewMaterial)
+{
+	if (ElementIndex == 0)
+	{
+		Plane.Material = NewMaterial;
+		UpdateMaterialParametersForMedia();
+	}
+}
+
+void UMediaPlaneComponent::GetUsedMaterials(TArray<UMaterialInterface*>& OutMaterials, bool bGetDebugMaterials) const
+{
+	OutMaterials.AddUnique(Plane.DynamicMaterial ? Plane.DynamicMaterial : Plane.Material);
+}
+
+void UMediaPlaneComponent::GetProjectionMatricesFromViewTarget(AActor* InViewTarget, FMatrix& OutViewProjectionMatrix, FMatrix& OutInvViewProjectionMatrix)
+{
+	check(InViewTarget);
+
+	FMinimalViewInfo MinimalViewInfo;
+	
+	ASceneCapture2D* SceneCapture = Cast<ASceneCapture2D>(InViewTarget);
+	USceneCaptureComponent2D* SceneCaptureComponent = SceneCapture ? SceneCapture->GetCaptureComponent2D() : nullptr;
+	if (SceneCaptureComponent)
+	{
+		MinimalViewInfo.Location = SceneCaptureComponent->GetComponentLocation();
+		MinimalViewInfo.Rotation = SceneCaptureComponent->GetComponentRotation();
+
+		MinimalViewInfo.FOV = SceneCaptureComponent->FOVAngle;
+		MinimalViewInfo.AspectRatio = SceneCaptureComponent->TextureTarget ? float(SceneCaptureComponent->TextureTarget->SizeX) / SceneCaptureComponent->TextureTarget->SizeY : 1.f;
+		MinimalViewInfo.bConstrainAspectRatio = false;
+		MinimalViewInfo.ProjectionMode = SceneCaptureComponent->ProjectionType;
+		MinimalViewInfo.OrthoWidth = SceneCaptureComponent->OrthoWidth;
+	}
+	else
+	{
+		InViewTarget->CalcCamera(0.f, MinimalViewInfo);
+	}
+
+	FMatrix ViewRotationMatrix = FInverseRotationMatrix(MinimalViewInfo.Rotation) * FMatrix(
+		FPlane(0, 0, 1, 0),
+		FPlane(1, 0, 0, 0),
+		FPlane(0, 1, 0, 0),
+		FPlane(0, 0, 0, 1));
+
+	FMatrix ProjectionMatrix;
+	if (SceneCaptureComponent && SceneCaptureComponent->bUseCustomProjectionMatrix)
+	{
+		ProjectionMatrix = AdjustProjectionMatrixForRHI(SceneCaptureComponent->CustomProjectionMatrix);
+	}
+	else
+	{
+		ProjectionMatrix = AdjustProjectionMatrixForRHI(CalculateProjectionMatrix(MinimalViewInfo));
+	}
+
+	FMatrix ViewMatrix = FTranslationMatrix(-MinimalViewInfo.Location) * ViewRotationMatrix;
+	FMatrix InvProjectionMatrix = ProjectionMatrix.Inverse();
+	FMatrix InvViewMatrix = ViewRotationMatrix.GetTransposed() * FTranslationMatrix(MinimalViewInfo.Location);
+
+	OutViewProjectionMatrix = ViewMatrix * ProjectionMatrix;
+	OutInvViewProjectionMatrix = InvProjectionMatrix * InvViewMatrix;
+}
+
+FBoxSphereBounds UMediaPlaneComponent::CalcBounds(const FTransform& LocalToWorld) const
+{
+	FBox MaxBox({
+		FVector(0,  1,  1),
+		FVector(0, -1,  1),
+		FVector(0,  1, -1),
+		FVector(0, -1, -1) });
+	MaxBox = MaxBox.TransformBy(LocalToWorld);
+	return FBoxSphereBounds(MaxBox);
+}
+
+AActor* UMediaPlaneComponent::FindViewTarget() const
+{
+	AActor* Actor = GetOwner();
+	while(Actor)
+	{
+		if (Actor->HasActiveCameraComponent() || Actor->FindComponentByClass<USceneCaptureComponent2D>())
+		{
+			return Actor;
+		}
+		Actor = Actor->GetAttachParentActor();
+	}
+
+	return nullptr;
+}
+
+#if WITH_EDITOR
+void UMediaPlaneComponent::PostEditChangeProperty(FPropertyChangedEvent& PropertyChangedEvent)
+{
+	Super::PostEditChangeProperty(PropertyChangedEvent);
+	UpdateMaterialParametersForMedia();
+	UpdateTransformScale();
+}
+
+void UMediaPlaneComponent::PostEditUndo()
+{
+	Super::PostEditUndo();
+	UpdateMaterialParametersForMedia();
+}
+
+UStructProperty* UMediaPlaneComponent::GetMediaPlaneProperty()
+{
+	return FindField<UStructProperty>(StaticClass(), GET_MEMBER_NAME_CHECKED(UMediaPlaneComponent, Plane));
+}
+
+#endif
\ No newline at end of file
diff --git a/Engine/Plugins/Media/MediaCompositing/Source/MediaCompositing/Private/Components/MediaPlaneFrustumComponent.cpp b/Engine/Plugins/Media/MediaCompositing/Source/MediaCompositing/Private/Components/MediaPlaneFrustumComponent.cpp
new file mode 100644
--- /dev/null
+++ b/Engine/Plugins/Media/MediaCompositing/Source/MediaCompositing/Private/Components/MediaPlaneFrustumComponent.cpp
@@ -0,0 +1,135 @@
+// Copyright 1998-2017 Epic Games, Inc. All Rights Reserved.
+
+/*=============================================================================
+	DrawFrsutumComponent.cpp: UDrawFrsutumComponent implementation.
+=============================================================================*/
+
+#include "MediaPlaneFrustumComponent.h"
+#include "PrimitiveViewRelevance.h"
+#include "PrimitiveSceneProxy.h"
+#include "Engine/CollisionProfile.h"
+#include "SceneManagement.h"
+#include "MediaPlaneComponent.h"
+
+
+/** Represents a draw frustum to the scene manager. */
+class FMediaPlaneFrustumSceneProxy : public FPrimitiveSceneProxy
+{
+public:
+
+	FMediaPlaneFrustumSceneProxy(const UMediaPlaneFrustumComponent* InComponent)
+		: FPrimitiveSceneProxy(InComponent)
+		, ViewTarget(nullptr)
+	{
+		bWillEverBeLit = false;
+
+		UMediaPlaneComponent* Parent = Cast<UMediaPlaneComponent>(InComponent->GetAttachParent());
+		if (Parent && Parent->GetPlane().bFillScreen)
+		{
+			ViewTarget = Parent->FindViewTarget();
+			InvViewProjectionMatrix = Parent->GetCachedInvViewProjectionMatrix();
+		}
+	}
+
+	virtual void GetDynamicMeshElements(const TArray<const FSceneView*>& Views, const FSceneViewFamily& ViewFamily, uint32 VisibilityMap, FMeshElementCollector& Collector) const override
+	{
+		if (!ViewTarget)
+		{
+			return;
+		}
+
+		for (int32 ViewIndex = 0; ViewIndex < Views.Num(); ViewIndex++)
+		{
+			if (!(VisibilityMap & (1 << ViewIndex)))
+			{
+				continue;
+			}
+
+			const FSceneView* View = Views[ViewIndex];
+
+			// Don't draw frustums on views that are set as the current view target
+			if (View->ViewActor == ViewTarget)
+			{
+				continue;
+			}
+
+			// These are all in world space
+			const float NearPlaneZ = 1.f;
+			FVector NearPlaneTopLeft 		= UMediaPlaneComponent::TransfromFromProjection(InvViewProjectionMatrix, FVector4(-1.f, -1.f, NearPlaneZ, 1.0f));
+			FVector NearPlaneTopRight 		= UMediaPlaneComponent::TransfromFromProjection(InvViewProjectionMatrix, FVector4( 1.f, -1.f, NearPlaneZ, 1.0f));
+			FVector NearPlaneBottomLeft 	= UMediaPlaneComponent::TransfromFromProjection(InvViewProjectionMatrix, FVector4(-1.f,  1.f, NearPlaneZ, 1.0f));
+			FVector NearPlaneBottomRight	= UMediaPlaneComponent::TransfromFromProjection(InvViewProjectionMatrix, FVector4( 1.f,  1.f, NearPlaneZ, 1.0f));
+
+			FVector DestinationTopLeft 		=  GetLocalToWorld().TransformPosition(FVector(0.f, -1.f, -1.f));
+			FVector DestinationTopRight 	=  GetLocalToWorld().TransformPosition(FVector(0.f,  1.f, -1.f));
+			FVector DestinationBottomLeft 	=  GetLocalToWorld().TransformPosition(FVector(0.f, -1.f,  1.f));
+			FVector DestinationBottomRight	=  GetLocalToWorld().TransformPosition(FVector(0.f,  1.f,  1.f));
+
+			FColor LineColor(255, 0, 255, 128);
+
+			FPrimitiveDrawInterface* PDI = Collector.GetPDI(ViewIndex);
+			const uint8 DepthPriorityGroup = GetDepthPriorityGroup(View);
+			PDI->DrawLine(NearPlaneTopLeft, DestinationTopLeft, LineColor, DepthPriorityGroup );
+			PDI->DrawLine(NearPlaneTopRight, DestinationTopRight, LineColor, DepthPriorityGroup );
+			PDI->DrawLine(NearPlaneBottomLeft, DestinationBottomLeft, LineColor, DepthPriorityGroup );
+			PDI->DrawLine(NearPlaneBottomRight, DestinationBottomRight, LineColor, DepthPriorityGroup );
+		}
+	}
+
+	virtual FPrimitiveViewRelevance GetViewRelevance(const FSceneView* View) const override
+	{
+		FPrimitiveViewRelevance Result;
+		Result.bDrawRelevance = IsShown(View);
+		Result.bDynamicRelevance = true;
+		Result.bShadowRelevance = IsShadowCast(View);
+		Result.bEditorPrimitiveRelevance = UseEditorCompositing(View);
+		return Result;
+	}
+
+	virtual uint32 GetMemoryFootprint( void ) const override { return( sizeof( *this ) + GetAllocatedSize() ); }
+	uint32 GetAllocatedSize( void ) const { return( FPrimitiveSceneProxy::GetAllocatedSize() ); }
+
+private:
+	AActor* ViewTarget;
+	FMatrix InvViewProjectionMatrix;
+};
+
+UMediaPlaneFrustumComponent::UMediaPlaneFrustumComponent(const FObjectInitializer& ObjectInitializer)
+	: Super(ObjectInitializer)
+{
+	bUseEditorCompositing = true;
+	bHiddenInGame = true;
+	SetCollisionProfileName(UCollisionProfile::NoCollision_ProfileName);
+	bGenerateOverlapEvents = false;
+}
+
+FPrimitiveSceneProxy* UMediaPlaneFrustumComponent::CreateSceneProxy()
+{
+	return new FMediaPlaneFrustumSceneProxy(this);
+}
+
+FBoxSphereBounds UMediaPlaneFrustumComponent::CalcBounds(const FTransform& LocalToWorld) const
+{
+	UMediaPlaneComponent* Parent = Cast<UMediaPlaneComponent>(GetAttachParent());
+	
+	FBox MaxBox({
+		FVector(0,  1,  1),
+		FVector(0, -1,  1),
+		FVector(0,  1, -1),
+		FVector(0, -1, -1),
+		});
+
+	MaxBox = MaxBox.TransformBy(LocalToWorld);
+
+	// Include the near view plane if possible
+	if (Parent && Parent->GetPlane().bFillScreen)
+	{
+		const FMatrix& InvViewProjectionMatrix = Parent->GetCachedInvViewProjectionMatrix();
+		
+		// These are all in world space
+		const float NearPlaneZ = 1.f;
+		MaxBox += UMediaPlaneComponent::TransfromFromProjection(InvViewProjectionMatrix, FVector4(0.f, 0.f, NearPlaneZ, 1.0f));
+	}
+
+	return FBoxSphereBounds(MaxBox);
+}
\ No newline at end of file
diff --git a/Engine/Plugins/Media/MediaCompositing/Source/MediaCompositing/Private/Components/MediaPlaneFrustumComponent.h b/Engine/Plugins/Media/MediaCompositing/Source/MediaCompositing/Private/Components/MediaPlaneFrustumComponent.h
new file mode 100644
--- /dev/null
+++ b/Engine/Plugins/Media/MediaCompositing/Source/MediaCompositing/Private/Components/MediaPlaneFrustumComponent.h
@@ -0,0 +1,26 @@
+// Copyright 1998-2016 Epic Games, Inc. All Rights Reserved.
+
+#pragma once
+
+#include "CoreMinimal.h"
+#include "UObject/ObjectMacros.h"
+#include "Components/PrimitiveComponent.h"
+#include "MediaPlaneFrustumComponent.generated.h"
+
+
+/** 
+ * A 2d material that will be rendered always facing the camera.
+ */
+UCLASS(MinimalAPI, ClassGroup=Rendering, hidecategories=(Object,Activation,Collision,"Components|Activation",Physics), editinlinenew)
+class UMediaPlaneFrustumComponent : public UPrimitiveComponent
+{
+public:
+	GENERATED_BODY()
+
+	UMediaPlaneFrustumComponent(const FObjectInitializer& Init);
+
+	//~ Begin UPrimitiveComponent Interface
+	virtual FPrimitiveSceneProxy* CreateSceneProxy() override;
+	virtual FBoxSphereBounds CalcBounds(const FTransform& LocalToWorld) const override;
+	virtual bool ShouldRecreateProxyOnUpdateTransform() const override { return true; }
+};
diff --git a/Engine/Plugins/Media/MediaCompositing/Source/MediaCompositing/Private/MediaCompositing.cpp b/Engine/Plugins/Media/MediaCompositing/Source/MediaCompositing/Private/MediaCompositing.cpp
new file mode 100644
--- /dev/null
+++ b/Engine/Plugins/Media/MediaCompositing/Source/MediaCompositing/Private/MediaCompositing.cpp
@@ -0,0 +1,6 @@
+// Copyright 1998-2017 Epic Games, Inc. All Rights Reserved.
+
+#include "Modules/ModuleManager.h"
+
+
+IMPLEMENT_MODULE(FDefaultModuleImpl, MediaCompositing);
diff --git a/Engine/Plugins/Media/MediaCompositing/Source/MediaCompositing/Private/MovieScene/MovieSceneMediaTemplate.cpp b/Engine/Plugins/Media/MediaCompositing/Source/MediaCompositing/Private/MovieScene/MovieSceneMediaTemplate.cpp
new file mode 100644
--- /dev/null
+++ b/Engine/Plugins/Media/MediaCompositing/Source/MediaCompositing/Private/MovieScene/MovieSceneMediaTemplate.cpp
@@ -0,0 +1,504 @@
+// Copyright 1998-2016 Epic Games, Inc. All Rights Reserved.
+
+#include "MovieSceneMediaTemplate.h"
+
+#include "MediaSource.h"
+#include "MediaPlayer.h"
+#include "MediaTexture.h"
+#include "Package.h"
+#include "UObject/GCObject.h"
+
+#include "MovieSceneMediaSection.h"
+#include "MovieSceneMediaTrack.h"
+
+
+/* Local helpers
+ *****************************************************************************/
+
+/**
+ * Simple struct that tries to keep an object rooted for this object's lifetime.
+ * Copying is prohibited. Moving passes ownership of the root to the LHS.
+ * Cannot not guarantee rooting outside of its control.
+ */
+template<typename T>
+struct TScopedRootObject
+{
+	/** Construct and roo the specified object */
+	TScopedRootObject(T* InObject)
+		: Object(InObject), bAddedToRoot(false)
+	{
+		if (InObject && !InObject->IsRooted())
+		{
+			Object->AddToRoot();
+		}
+	}
+
+	/** Copying is disabled */
+	TScopedRootObject(const TScopedRootObject&) = delete;
+	TScopedRootObject& operator=(const TScopedRootObject&) = delete;
+
+	/** Move construction/assignment passes ownership of the root to the LHS */
+	TScopedRootObject(TScopedRootObject&& In)
+		: Object(In.Object)
+		, bAddedToRoot(In.bAddedToRoot)
+	{
+		In.Object = nullptr;
+		In.bAddedToRoot = false;
+	}
+
+	TScopedRootObject& operator=(TScopedRootObject&& In)
+	{
+		Object = In.Object;
+		bAddedToRoot = In.bAddedToRoot;
+		
+		In.Object = nullptr;
+		In.bAddedToRoot = false;
+
+		return *this;
+	}
+
+	/** Destructor frees the root, if necessary */
+	~TScopedRootObject()
+	{
+		if (Object && bAddedToRoot)
+		{
+			Object->RemoveFromRoot();
+		}
+	}
+
+	T* GetObject() const
+	{
+		return Object;
+	}
+
+private:
+
+	/** The object that should be rooted */
+	T* Object;
+
+	/** Whether we added the object to the root (false if it's null, or was already rooted) */
+	bool bAddedToRoot;
+};
+
+
+// This pre-animated token is only created if we create a new mediaplayer and assign it to the property
+struct FMediaPlayerPropertyPreAnimatedToken
+	: IMovieScenePreAnimatedToken
+{
+	FMediaPlayerPropertyPreAnimatedToken(UMediaPlayer* InPreviousPropertyValue, const TSharedPtr<FTrackInstancePropertyBindings>& InBindings)
+		: PreviousPropertyValue(InPreviousPropertyValue), PropertyBindings(InBindings)
+	{ }
+
+	/** Manual defaulted functions while we need to support compilers that can't generate/default them for us */
+#if PLATFORM_COMPILER_HAS_DEFAULTED_FUNCTIONS
+	FMediaPlayerPropertyPreAnimatedToken(FMediaPlayerPropertyPreAnimatedToken&& In) = default;
+	FMediaPlayerPropertyPreAnimatedToken& operator=(FMediaPlayerPropertyPreAnimatedToken&& In) = default;
+
+#else
+	FMediaPlayerPropertyPreAnimatedToken(FMediaPlayerPropertyPreAnimatedToken&& In)
+		: PreviousPropertyValue(MoveTemp(In.PreviousPropertyValue))
+		, PropertyBindings(MoveTemp(In.PropertyBindings))
+	{ }
+
+	FMediaPlayerPropertyPreAnimatedToken& operator=(FMediaPlayerPropertyPreAnimatedToken&& In)
+	{
+		PreviousPropertyValue = MoveTemp(In.PreviousPropertyValue);
+		PropertyBindings = MoveTemp(In.PropertyBindings);
+
+		return *this;
+	}
+#endif
+
+	/** Virtual destructor. */
+	virtual ~FMediaPlayerPropertyPreAnimatedToken() { }
+
+public:
+
+	virtual void RestoreState(UObject& RestoreObject, IMovieScenePlayer& Player)
+	{
+		UMediaPlayer* PreviousPlayer = PreviousPropertyValue.GetObject();
+		PropertyBindings->CallFunction<UMediaPlayer*>(RestoreObject, PreviousPlayer);
+	}
+
+private:
+
+	/** Keep the previous player alive so we can definitely restore it later on */
+	TScopedRootObject<UMediaPlayer> PreviousPropertyValue;
+
+	/** Property bindings that allow us to set the property when we've finished evaluating */
+	TSharedPtr<FTrackInstancePropertyBindings> PropertyBindings;
+
+public:
+
+	/** Producer that can create these tokens */
+	struct FProducer
+		: IMovieScenePreAnimatedTokenProducer
+	{
+		FProducer(const TSharedPtr<FTrackInstancePropertyBindings>& InBindings)
+			: Bindings(InBindings)
+		{ }
+
+		virtual ~FProducer() { }
+
+		virtual IMovieScenePreAnimatedTokenPtr CacheExistingState(UObject& Object) const override
+		{
+			return FMediaPlayerPropertyPreAnimatedToken(Bindings->GetCurrentValue<UMediaPlayer*>(Object), Bindings);
+		}
+
+		TSharedPtr<FTrackInstancePropertyBindings> Bindings;
+	};
+};
+
+
+/**
+ * This pre-animated token is only created if we create a new
+ * media texture and assign it to an existing media player.
+ */
+struct FMediaPlayerTexturePreAnimatedToken
+	: IMovieScenePreAnimatedToken
+{
+	/** Virtual destructor. */
+	virtual ~FMediaPlayerTexturePreAnimatedToken() { }
+
+	virtual void RestoreState(UObject& RestoreObject, IMovieScenePlayer& Player)
+	{
+		UMediaPlayer* MediaPlayer = CastChecked<UMediaPlayer>(&RestoreObject);
+		// @todo gmp: fix me media framework 3.0
+		/*
+		MediaPlayer->SetVideoTexture(nullptr);
+		*/
+	}
+
+	static FMovieSceneAnimTypeID GetAnimTypeID() { return TMovieSceneAnimTypeID<FMediaPlayerTexturePreAnimatedToken>(); }
+
+	struct FProducer : IMovieScenePreAnimatedTokenProducer
+	{
+		virtual IMovieScenePreAnimatedTokenPtr CacheExistingState(UObject& Object) const override
+		{
+			return FMediaPlayerTexturePreAnimatedToken();
+		}
+	};
+};
+
+
+/**
+ * Persistent data that's stored for each currently evaluating section.
+ */
+struct FMediaSectionData
+	: PropertyTemplate::FSectionData, FGCObject
+{
+	FMediaSectionData()
+		: TemporaryMediaPlayer(nullptr)
+	{ }
+
+	virtual void AddReferencedObjects(FReferenceCollector& Collector) override
+	{
+		Collector.AddReferencedObject(TemporaryMediaPlayer);
+	}
+
+	UMediaPlayer* GetTemporaryMediaPlayer()
+	{
+		if (TemporaryMediaPlayer == nullptr)
+		{
+			TemporaryMediaPlayer = NewObject<UMediaPlayer>(GetTransientPackage(), MakeUniqueObjectName(GetTransientPackage(), UMediaPlayer::StaticClass()));
+		}
+
+		return TemporaryMediaPlayer;
+	}
+
+	UMediaPlayer* GetOrUpdateMediaPlayerFromProperty(UObject& InObject, IMovieScenePlayer& Player)
+	{
+		UMediaPlayer* MediaPlayer = PropertyBindings->GetCurrentValue<UMediaPlayer*>(InObject);
+
+		// assign the temporary player to the property if it's null or not the temporary player
+		const bool bNeedToAssignTemporaryPlayer = !MediaPlayer || (TemporaryMediaPlayer && MediaPlayer != TemporaryMediaPlayer);
+
+		if (bNeedToAssignTemporaryPlayer)
+		{
+			GetTemporaryMediaPlayer();
+			MediaPlayer = TemporaryMediaPlayer;
+		}
+
+		check(MediaPlayer);
+
+		// Ensure we have a video texture on the media player (transient ones will always
+		// have a texture already, but media players specified on external properties may not)
+
+		// @todo gmp: fix me Media Framework 3.0
+		/*
+		if (!MediaPlayer->GetVideoTexture())
+		{
+			UMediaTexture* NewTexture = NewObject<UMediaTexture>(GetTransientPackage(), MakeUniqueObjectName(GetTransientPackage(), UMediaTexture::StaticClass()));
+			NewTexture->UpdateResource();
+
+			// Save the previous video texture so we don't end up trying to save one that's in a transient package
+			Player.SavePreAnimatedState(*MediaPlayer, FMediaPlayerTexturePreAnimatedToken::GetAnimTypeID(), FMediaPlayerTexturePreAnimatedToken::FProducer());
+
+			MediaPlayer->SetVideoTexture(NewTexture);
+		}
+		*/
+
+		// always assign the property last, so that anything responding to the update
+		// can have access to the media player and its video texture if needs be
+
+		if (bNeedToAssignTemporaryPlayer)
+		{
+			// save the previous property value using a unique anim type ID to ensure we don't leave it set after evaluation
+			Player.SavePreAnimatedState(InObject, PropertyID, FMediaPlayerPropertyPreAnimatedToken::FProducer(PropertyBindings));
+
+			// assign the property
+			PropertyBindings->CallFunction<UMediaPlayer*>(InObject, MediaPlayer);
+		}
+
+		return MediaPlayer;
+	}
+
+private:
+
+	UMediaPlayer* TemporaryMediaPlayer;
+};
+
+
+struct FMediaPlayerPreAnimatedTokenProducer
+	: IMovieScenePreAnimatedTokenProducer
+{
+	virtual IMovieScenePreAnimatedTokenPtr CacheExistingState(UObject& Object) const override
+	{
+		struct FToken : IMovieScenePreAnimatedToken
+		{
+			FToken(float InRate, const FString& InOldUrl, FTimespan InSeekPosition)
+				: Rate(InRate), OldUrl(InOldUrl), SeekPosition(InSeekPosition)
+			{}
+
+			virtual void RestoreState(UObject& RestoreObject, IMovieScenePlayer& Player)
+			{
+				UMediaPlayer* LocalMediaPlayer = CastChecked<UMediaPlayer>(&RestoreObject);
+
+				if (!OldUrl.IsEmpty())
+				{
+					LocalMediaPlayer->OpenUrl(OldUrl);
+				}
+				else
+				{
+					LocalMediaPlayer->Close();
+				}
+
+				LocalMediaPlayer->SetRate(Rate);
+				LocalMediaPlayer->Seek(SeekPosition);
+			}
+
+			float Rate;
+			FString OldUrl;
+			FTimespan SeekPosition;
+		};
+
+		UMediaPlayer* MediaPlayer = CastChecked<UMediaPlayer>(&Object);
+		return FToken{ MediaPlayer->GetRate(), MediaPlayer->GetUrl(), MediaPlayer->GetTime() };
+	}
+};
+
+
+struct FMediaSectionPreRollExecutionToken
+	: IMovieSceneExecutionToken
+{
+	FMediaSectionPreRollExecutionToken(UMediaSource* InSource, FTimespan InMediaTime)
+		: Source(InSource)
+		, MediaTime(InMediaTime)
+	{
+	}
+
+	virtual void Execute(const FMovieSceneContext& Context, const FMovieSceneEvaluationOperand& Operand, FPersistentEvaluationData& PersistentData, IMovieScenePlayer& Player) override
+	{
+		using namespace PropertyTemplate;
+
+		const bool bForceTemporaryPlayer = ForceTemporaryPlayer();
+
+		FMediaSectionData& SectionData = PersistentData.GetSectionData<FMediaSectionData>();
+		for (TWeakObjectPtr<> WeakObject : Player.FindBoundObjects(Operand))
+		{
+			UObject* Object = WeakObject.Get();
+			if (!Object)
+			{
+				continue;
+			}
+
+			UMediaPlayer* MediaPlayer = bForceTemporaryPlayer ? SectionData.GetTemporaryMediaPlayer() : SectionData.GetOrUpdateMediaPlayerFromProperty(*Object, Player);
+			if (!ensure(MediaPlayer))
+			{
+				continue;
+			}
+
+			// Save the previous media player state
+			Player.SavePreAnimatedState(*MediaPlayer, TMovieSceneAnimTypeID<FMediaSectionPreRollExecutionToken>(), FMediaPlayerPreAnimatedTokenProducer());
+
+			// Open the media source if necessary
+			if (MediaPlayer->GetUrl() != Source->GetUrl())
+			{
+				if (MediaPlayer->CanPlaySource(Source))
+				{
+					MediaPlayer->OpenSource(Source);
+				}
+				else
+				{
+					// @todo: log warning/error
+				}
+			}
+
+			OnUpdateFrame(MediaPlayer, Context);
+		}
+	}
+
+protected:
+
+	virtual bool ForceTemporaryPlayer() const
+	{
+		return true;
+	}
+
+	virtual void OnUpdateFrame(UMediaPlayer* MediaPlayer, const FMovieSceneContext& Context)
+	{
+		if (MediaPlayer->GetRate() != 0.f)
+		{
+			MediaPlayer->SetRate(0.f);
+		}
+
+		// @todo: Set this on the section itself
+		if (MediaPlayer->GetTime() != MediaTime)
+		{
+			MediaPlayer->Seek(MediaTime);
+		}
+	}
+
+	UMediaSource* Source;
+	FTimespan MediaTime;
+};
+
+
+struct FMediaSectionExecutionToken
+	: FMediaSectionPreRollExecutionToken
+{
+	FMediaSectionExecutionToken(UMediaSource* InSource, FTimespan InMediaTime)
+		: FMediaSectionPreRollExecutionToken(InSource, InMediaTime)
+		, PlaybackRate(1.f)
+	{
+	}
+
+	virtual bool ForceTemporaryPlayer() const
+	{
+		return false;
+	}
+
+	virtual void OnUpdateFrame(UMediaPlayer* MediaPlayer, const FMovieSceneContext& Context) override
+	{
+		// Play and update the media player
+		if (Context.GetStatus() == EMovieScenePlayerStatus::Playing)
+		{
+			if (MediaPlayer->GetRate() != PlaybackRate)
+			{
+				MediaPlayer->SetRate(PlaybackRate);
+				MediaPlayer->Seek(MediaTime);
+			}
+		}
+		else
+		{
+			MediaPlayer->SetRate(0.0f);
+			MediaPlayer->Seek(MediaTime);
+		}
+	}
+
+	float PlaybackRate;
+};
+
+
+/* FMovieSceneMediaSectionTemplate structors
+ *****************************************************************************/
+
+FMovieSceneMediaSectionTemplate::FMovieSceneMediaSectionTemplate(const UMovieSceneMediaSection& InSection, const UMovieSceneMediaTrack& InTrack)
+	: PropertyData(InTrack.GetPropertyName(), InTrack.GetPropertyPath(), NAME_None, "OnMediaPlayerChanged")
+{
+	Params.Source = InSection.GetMediaSource();
+	Params.SectionStartTime = InSection.GetStartTime();
+}
+
+
+/* FMovieSceneEvalTemplate interface
+ *****************************************************************************/
+
+void FMovieSceneMediaSectionTemplate::Evaluate(const FMovieSceneEvaluationOperand& Operand, const FMovieSceneContext& Context, const FPersistentEvaluationData& PersistentData, FMovieSceneExecutionTokens& ExecutionTokens) const
+{
+	if (!Params.Source)
+	{
+		return;
+	}
+
+	if (Context.IsPreRoll())
+	{
+		// Use the end of the preroll range if we're in shot preroll, otherwise start at frame 0
+		// @todo: account for start offset and video time dilation if/when these are added
+		float FirstMediaFrameTime = Context.HasPreRollEndTime() ? Context.GetPreRollEndTime() - Params.SectionStartTime : 0.f;
+		ExecutionTokens.Add(FMediaSectionPreRollExecutionToken(Params.Source, FTimespan(int64(ETimespan::TicksPerSecond * FirstMediaFrameTime))));
+	}
+	else if (!Context.IsPostRoll())
+	{
+		// Only update the video if we're not postrolling (which only should happen when playing in reverse through the preroll range)
+
+		// @todo: account for start offset and video time dilation if/when these are added
+		float MediaTime = Context.GetTime() - Params.SectionStartTime;
+		ExecutionTokens.Add(FMediaSectionExecutionToken(Params.Source, FTimespan(int64(ETimespan::TicksPerSecond * MediaTime))));
+	}
+}
+
+
+UScriptStruct& FMovieSceneMediaSectionTemplate::GetScriptStructImpl() const
+{
+	return *StaticStruct();
+}
+
+
+void FMovieSceneMediaSectionTemplate::Initialize(const FMovieSceneEvaluationOperand& Operand, const FMovieSceneContext& Context, FPersistentEvaluationData& PersistentData, IMovieScenePlayer& Player) const
+{
+	// Null out any prerolling video textures right at the start of the frame, before any other video sections update through their execution tokens.
+	if (!Context.IsPreRoll())
+	{
+		return;
+	}
+
+	FMediaSectionData& SectionData = PersistentData.GetSectionData<FMediaSectionData>();
+	UMediaPlayer* TemporaryMediaPlayer = SectionData.GetTemporaryMediaPlayer();
+
+	for (TWeakObjectPtr<> WeakObject : Player.FindBoundObjects(Operand))
+	{
+		if (UObject* Object = WeakObject.Get())
+		{
+			UMediaPlayer* MediaPlayerFromProperty = SectionData.PropertyBindings->GetCurrentValue<UMediaPlayer*>(*Object);
+
+			// Ensure the temporary (prerolling) player has no video texture. It will get
+			// created when the video is actually evaluated for real. This will only be
+			// set if we've scrubbed backwards out of a video section into its pre-roll
+			
+			// @todo gmp: fix me media framework 3.0
+			/*
+			TemporaryMediaPlayer->SetVideoTexture(nullptr);
+			*/
+
+			if (MediaPlayerFromProperty == TemporaryMediaPlayer)
+			{
+				// Ensure the property is not set if we're prerolling (so we don't display preroll frames)
+				FMovieSceneAnimTypeID PropertyID = SectionData.PropertyID;
+				Player.RestorePreAnimatedState(*Object, [=](FMovieSceneAnimTypeID In){ return In == PropertyID; });
+			}
+		}
+	}
+}
+
+
+void FMovieSceneMediaSectionTemplate::Setup(FPersistentEvaluationData& PersistentData, IMovieScenePlayer& Player) const
+{
+	PropertyData.SetupTrack<FMediaSectionData>(PersistentData);
+}
+
+
+void FMovieSceneMediaSectionTemplate::SetupOverrides()
+{
+	EnableOverrides(RequiresSetupFlag | RequiresInitializeFlag);
+}
diff --git a/Engine/Plugins/Media/MediaCompositing/Source/MediaCompositing/Public/MediaPlane.h b/Engine/Plugins/Media/MediaCompositing/Source/MediaCompositing/Public/MediaPlane.h
new file mode 100644
--- /dev/null
+++ b/Engine/Plugins/Media/MediaCompositing/Source/MediaCompositing/Public/MediaPlane.h
@@ -0,0 +1,45 @@
+// Copyright 1998-2016 Epic Games, Inc. All Rights Reserved.
+
+#pragma once
+
+#include "UObject/ObjectMacros.h"
+#include "GameFramework/Actor.h"
+
+#include "MediaPlane.generated.h"
+
+class UMediaPlaneComponent;
+
+
+UCLASS(ClassGroup=Rendering, hidecategories=(Object, Activation, Physics, Collision, Input, PhysicsVolume))
+class MEDIACOMPOSITING_API AMediaPlane
+	: public AActor
+{
+public:
+
+	GENERATED_BODY()
+
+	/**
+	 * Create and initialize a new instance.
+	 *
+	 * @param ObjectInitializer Object initializer.
+	 */
+	AMediaPlane(const FObjectInitializer& ObjectInitializer);
+
+public:
+
+	/**
+	 * Get the media plane component of this actor.
+	 *
+	 * @return The media plane component.
+	 */
+	UMediaPlaneComponent* GetMediaPlaneComponent() const
+	{
+		return MediaPlane;
+	}
+
+protected:
+
+	/** The media plane component used by this actor. */
+	UPROPERTY(Category=MediaPlane, VisibleAnywhere, BlueprintReadOnly, meta=(ExposeFunctionCategories="Mesh,Rendering,Physics,Components|StaticMesh"))
+	UMediaPlaneComponent* MediaPlane;
+};
diff --git a/Engine/Plugins/Media/MediaCompositing/Source/MediaCompositing/Public/MediaPlaneComponent.h b/Engine/Plugins/Media/MediaCompositing/Source/MediaCompositing/Public/MediaPlaneComponent.h
new file mode 100644
--- /dev/null
+++ b/Engine/Plugins/Media/MediaCompositing/Source/MediaCompositing/Public/MediaPlaneComponent.h
@@ -0,0 +1,151 @@
+// Copyright 1998-2016 Epic Games, Inc. All Rights Reserved.
+
+#pragma once
+
+#include "CoreMinimal.h"
+#include "UObject/ObjectMacros.h"
+#include "Components/PrimitiveComponent.h"
+#include "MediaPlaneComponent.generated.h"
+
+class FPrimitiveSceneProxy;
+class UMediaPlaneFrustumComponent;
+class UMaterialInterface;
+class UTexture;
+
+USTRUCT(BlueprintType)
+struct MEDIACOMPOSITING_API FMediaPlaneParameters
+{
+	GENERATED_BODY()
+
+	FMediaPlaneParameters();
+
+	/** The material that the media plane is rendered with. */
+	UPROPERTY(EditAnywhere, BlueprintReadWrite, AdvancedDisplay, Category="Media Plane")
+	UMaterialInterface* Material;
+
+	/** Name of a texture parameter inside the material to patch the render target texture to */
+	UPROPERTY(EditAnywhere, BlueprintReadWrite, AdvancedDisplay, Category="Media Plane")
+	FName TextureParameterName;
+
+	/** Automaticaly size the plane based on the active camera's lens and filmback settings. Target Camera is found by looking for an active camera component from this component's actor, through its attached parents. */
+	UPROPERTY(EditAnywhere, BlueprintReadWrite, Category="Media Plane")
+	bool bFillScreen;
+
+	/** The amount to fill the screen with when attached to a camera */
+	UPROPERTY(EditAnywhere, BlueprintReadWrite, Category="Media Plane", meta=(EditCondition="bFillScreen", AllowPreserveRatio))
+	FVector2D FillScreenAmount;
+
+	/** The fixed size of the media plane */
+	UPROPERTY(EditAnywhere, BlueprintReadWrite, Category="Media Plane", meta=(EditCondition="!bFillScreen"))
+	FVector2D FixedSize;
+
+	/** Transient texture that receives image frames */
+	UPROPERTY(transient, BlueprintReadOnly, Category="Media Plane")
+	UTexture* RenderTexture;
+
+	/** Transient MID to hold the material with the render texture patched in */
+	UPROPERTY(transient)
+	UMaterialInstanceDynamic* DynamicMaterial;
+};
+
+/** 
+ * A 2d plane that will be rendered always facing the camera.
+ */
+UCLASS(ClassGroup=Rendering, hidecategories=(Object,Activation,Collision,"Components|Activation",Physics), editinlinenew, meta=(BlueprintSpawnableComponent))
+class MEDIACOMPOSITING_API UMediaPlaneComponent : public UPrimitiveComponent
+{
+public:
+
+	GENERATED_BODY()
+	UMediaPlaneComponent(const FObjectInitializer& Init);
+
+	static inline FVector TransfromFromProjection(const FMatrix& Matrix, const FVector4& InVector)
+	{
+		FVector4 HomogenousVector = Matrix.TransformFVector4(InVector);
+		FVector ReturnVector = HomogenousVector;
+		if (HomogenousVector.W != 0.0f)
+		{
+			ReturnVector /= HomogenousVector.W;
+		}
+
+		return ReturnVector;
+	}
+
+	/** Add an media mlane to this actor */
+	UFUNCTION(BlueprintCallable, Category="Game|Media Plane")
+	void SetMediaPlane(FMediaPlaneParameters Plane);
+
+	/** Get this actor's media planes */
+	UFUNCTION(BlueprintCallable, Category="Game|Media Plane")
+	FMediaPlaneParameters GetPlane() const
+	{
+		return Plane;
+	}
+
+	/** Called by sequencer if a texture is changed */
+	UFUNCTION()
+	void OnRenderTextureChanged();
+
+	/** Access this component's cached view projection matrix. Only valid when the plane is set to fill screen */
+	const FMatrix& GetCachedViewProjectionMatrix() const
+	{
+		return ViewProjectionMatrix;
+	}
+
+	/** Access this component's cached inverse view projection matrix. Only valid when the plane is set to fill screen */
+	const FMatrix& GetCachedInvViewProjectionMatrix() const
+	{
+		return InvViewProjectionMatrix;
+	}
+
+	//~ Begin UPrimitiveComponent Interface
+	virtual void TickComponent(float DeltaTime, enum ELevelTick TickType, FActorComponentTickFunction *ThisTickFunction) override;
+	virtual FPrimitiveSceneProxy* CreateSceneProxy() override;
+	virtual void OnUpdateTransform(EUpdateTransformFlags UpdateTransformFlags, ETeleportType Teleport) override;
+	virtual FBoxSphereBounds CalcBounds(const FTransform& LocalToWorld) const override;
+	virtual UMaterialInterface* GetMaterial(int32 Index) const override;
+	virtual void GetUsedMaterials(TArray<UMaterialInterface*>& OutMaterials, bool bGetDebugMaterials) const override;
+	virtual void SetMaterial(int32 ElementIndex, UMaterialInterface* Material) override;
+	virtual void OnRegister() override;
+#if WITH_EDITOR
+	virtual void PostEditChangeProperty(FPropertyChangedEvent& PropertyChangedEvent) override;
+	virtual void PostEditUndo() override;
+#endif
+
+#if WITH_EDITOR
+	/** Access the property relating to this component's Media Plane */
+	static UStructProperty* GetMediaPlaneProperty();
+#endif
+
+	/**
+	 * Finds a view target that this Media Plane is presenting to
+	 */
+	AActor* FindViewTarget() const;
+
+	/**
+	 * Calculate projection matrices from a specified view target
+	 */
+	static void GetProjectionMatricesFromViewTarget(AActor* InViewTarget, FMatrix& OutViewProjectionMatrix, FMatrix& OutInvViewProjectionMatrix);
+
+protected:
+
+	void UpdateMaterialParametersForMedia();
+
+	void UpdateTransformScale();
+
+private:
+
+	/** Array of Media Planes to render for this component */
+	UPROPERTY(EditAnywhere, Category="Media Plane", DisplayName="Media Planes", meta=(ShowOnlyInnerProperties))
+	FMediaPlaneParameters Plane;
+
+	FMatrix ViewProjectionMatrix;
+	FMatrix InvViewProjectionMatrix;
+
+	bool bReenetrantTranformChange;
+
+#if WITH_EDITORONLY_DATA
+	UPROPERTY(transient)
+	UMediaPlaneFrustumComponent* EditorFrustum;
+#endif
+};
diff --git a/Engine/Plugins/Media/MediaCompositing/Source/MediaCompositing/Public/MovieSceneMediaTrack.h b/Engine/Plugins/Media/MediaCompositing/Source/MediaCompositing/Public/MovieSceneMediaTrack.h
new file mode 100644
--- /dev/null
+++ b/Engine/Plugins/Media/MediaCompositing/Source/MediaCompositing/Public/MovieSceneMediaTrack.h
@@ -0,0 +1,45 @@
+// Copyright 1998-2016 Epic Games, Inc. All Rights Reserved.
+
+#pragma once
+
+#include "Misc/Guid.h"
+#include "MovieScenePropertyTrack.h"
+#include "UObject/ObjectMacros.h"
+
+#include "MovieSceneMediaTrack.generated.h"
+
+
+/**
+ * Implements a movie scene track for media playback.
+ */
+UCLASS(MinimalAPI)
+class UMovieSceneMediaTrack
+	: public UMovieScenePropertyTrack
+{
+public:
+
+	GENERATED_BODY()
+
+	/**
+	 * Create and initialize a new instance.
+	 *
+	 * @param ObjectInitializer The object initializer.
+	 */
+	UMovieSceneMediaTrack(const FObjectInitializer& ObjectInitializer);
+
+public:
+
+	//~ UMovieScenePropertyTrack interface
+
+	virtual void AddSection(UMovieSceneSection& Section) override;
+	virtual UMovieSceneSection* CreateNewSection() override;
+	virtual FMovieSceneEvalTemplatePtr CreateTemplateForSection(const UMovieSceneSection& InSection) const override;
+	virtual const TArray<UMovieSceneSection*>& GetAllSections() const override;
+	virtual void RemoveSection(UMovieSceneSection& Section) override;
+	virtual bool SupportsMultipleRows() const override { return true; }
+
+#if WITH_EDITORONLY_DATA
+	virtual FText GetDefaultDisplayName() const override;
+	virtual FName GetTrackName() const override;
+#endif
+};
diff --git a/Engine/Plugins/Media/MediaCompositing/Source/MediaCompositingEditor/MediaCompositingEditor.Build.cs b/Engine/Plugins/Media/MediaCompositing/Source/MediaCompositingEditor/MediaCompositingEditor.Build.cs
new file mode 100644
--- /dev/null
+++ b/Engine/Plugins/Media/MediaCompositing/Source/MediaCompositingEditor/MediaCompositingEditor.Build.cs
@@ -0,0 +1,36 @@
+// Copyright 1998-2016 Epic Games, Inc. All Rights Reserved.
+
+namespace UnrealBuildTool.Rules
+{
+	public class MediaCompositingEditor : ModuleRules
+	{
+		public MediaCompositingEditor(ReadOnlyTargetRules Target) : base(Target)
+		{
+			PrivateDependencyModuleNames.AddRange(
+				new string[]
+				{
+					"Core",
+					"CoreUObject",
+					"EditorStyle",
+					"Engine",
+					"MediaAssets",
+					"MediaCompositing",
+					"MovieScene",
+					"MovieSceneTools",
+					"MovieSceneTracks",
+					"RHI",
+					"Sequencer",
+					"Slate",
+					"SlateCore",
+					"UnrealEd",
+				});
+
+			PrivateIncludePaths.AddRange(
+				new string[] {
+					"MediaCompositingEditor/Private",
+					"MediaCompositingEditor/Private/Sequencer",
+					"MediaCompositingEditor/Private/Shared",
+				});
+		}
+	}
+}
