Copying //UE4/Dev-Sequencer to //UE4/Main (Source: //UE4/Dev-Sequencer @ 3617642)

diff --git a/Engine/Plugins/Media/MediaCompositing/Source/MediaCompositingEditor/Private/Sequencer/MediaThumbnailSection.cpp b/Engine/Plugins/Media/MediaCompositing/Source/MediaCompositingEditor/Private/Sequencer/MediaThumbnailSection.cpp
new file mode 100644
--- /dev/null
+++ b/Engine/Plugins/Media/MediaCompositing/Source/MediaCompositingEditor/Private/Sequencer/MediaThumbnailSection.cpp
@@ -0,0 +1,239 @@
+// Copyright 1998-2016 Epic Games, Inc. All Rights Reserved.
+
+#include "MediaThumbnailSection.h"
+
+#include "EditorStyleSet.h"
+#include "ISequencer.h"
+#include "MediaPlayer.h"
+#include "MediaSource.h"
+#include "MediaTexture.h"
+#include "MovieSceneMediaSection.h"
+#include "Rendering/DrawElements.h"
+#include "SequencerSectionPainter.h"
+#include "TrackEditorThumbnail/TrackEditorThumbnailPool.h"
+#include "UObject/Package.h"
+#include "UObject/UObjectGlobals.h"
+#include "Widgets/Layout/SBox.h"
+#include "Widgets/Text/STextBlock.h"
+
+
+#define LOCTEXT_NAMESPACE "FMediaThumbnailSection"
+
+
+/* FMediaSection structors
+ *****************************************************************************/
+
+FMediaThumbnailSection::FMediaThumbnailSection(UMovieSceneMediaSection& InSection, TSharedPtr<FTrackEditorThumbnailPool> InThumbnailPool, TSharedPtr<ISequencer> InSequencer)
+	: FThumbnailSection(InSequencer, InThumbnailPool, this, InSection)
+{
+	MediaPlayer = nullptr;
+	MediaTexture = nullptr;
+	TimeSpace = ETimeSpace::Local;
+}
+
+
+FMediaThumbnailSection::~FMediaThumbnailSection()
+{
+	if (MediaPlayer != nullptr)
+	{
+		MediaPlayer->Close();
+	}
+}
+
+
+/* FGCObject interface
+ *****************************************************************************/
+
+void FMediaThumbnailSection::AddReferencedObjects(FReferenceCollector& Collector)
+{
+	Collector.AddReferencedObject(MediaPlayer);
+	Collector.AddReferencedObject(MediaTexture);
+}
+
+
+/* FThumbnailSection interface
+ *****************************************************************************/
+
+FText FMediaThumbnailSection::GetSectionTitle() const
+{
+	if (UMediaSource* MediaSource = Section ? CastChecked<UMovieSceneMediaSection>(Section)->GetMediaSource() : nullptr)
+	{
+		return FText::FromString(MediaSource->GetName());
+	}
+
+	return FText::GetEmpty();
+}
+
+
+void FMediaThumbnailSection::SetSingleTime(float GlobalTime)
+{
+	UMovieSceneMediaSection* MediaSection = CastChecked<UMovieSceneMediaSection>(Section);
+	
+	if (MediaSection)
+	{
+		MediaSection->SetThumbnailReferenceOffset(GlobalTime - MediaSection->GetStartTime());
+	}
+}
+
+
+TSharedRef<SWidget> FMediaThumbnailSection::GenerateSectionWidget()
+{
+	return SNew(SBox)
+		.HAlign(HAlign_Left)
+		.VAlign(VAlign_Top)
+		.Padding(GetContentPadding())
+		[
+			SNew(STextBlock)
+				.Text(this, &FMediaThumbnailSection::GetSectionTitle)
+				.ShadowOffset(FVector2D(1.0f, 1.0f))
+		];
+}
+
+
+float FMediaThumbnailSection::GetSectionHeight() const
+{
+	return FThumbnailSection::GetSectionHeight() + 2 * 9.0f; // make space for the film border
+}
+
+
+FMargin FMediaThumbnailSection::GetContentPadding() const
+{
+	return FMargin(8.0f, 15.0f);
+}
+
+
+void FMediaThumbnailSection::Tick(const FGeometry& AllottedGeometry, const FGeometry& ClippedGeometry, const double InCurrentTime, const float InDeltaTime)
+{
+	UMovieSceneMediaSection* MediaSection = Cast<UMovieSceneMediaSection>(Section);
+
+	if (MediaSection != nullptr)
+	{
+		if (GetDefault<UMovieSceneUserThumbnailSettings>()->bDrawSingleThumbnails)
+		{
+			ThumbnailCache.SetSingleReferenceFrame(MediaSection->GetThumbnailReferenceOffset());
+		}
+		else
+		{
+			ThumbnailCache.SetSingleReferenceFrame(TOptional<float>());
+		}
+	}
+
+	FThumbnailSection::Tick(AllottedGeometry, ClippedGeometry, InCurrentTime, InDeltaTime);
+}
+
+
+int32 FMediaThumbnailSection::OnPaintSection(FSequencerSectionPainter& InPainter) const
+{
+	static const FSlateBrush* FilmBorder = FEditorStyle::GetBrush("Sequencer.Section.FilmBorder");
+
+	InPainter.LayerId = InPainter.PaintSectionBackground();
+
+	FVector2D LocalSectionSize = InPainter.SectionGeometry.GetLocalSize();
+
+	FSlateClippingZone ClippingZone(InPainter.SectionClippingRect.InsetBy(FMargin(1.0f)));
+	InPainter.DrawElements.PushClip(ClippingZone);
+
+	FSlateDrawElement::MakeBox(
+		InPainter.DrawElements,
+		InPainter.LayerId++,
+		InPainter.SectionGeometry.ToPaintGeometry(FVector2D(LocalSectionSize.X - 2.0f, 7.0f), FSlateLayoutTransform(FVector2D(1.0f, 4.0f))),
+		FilmBorder,
+		InPainter.bParentEnabled ? ESlateDrawEffect::None : ESlateDrawEffect::DisabledEffect
+	);
+
+	FSlateDrawElement::MakeBox(
+		InPainter.DrawElements,
+		InPainter.LayerId++,
+		InPainter.SectionGeometry.ToPaintGeometry(FVector2D(LocalSectionSize.X - 2.0f, 7.0f), FSlateLayoutTransform(FVector2D(1.0f, LocalSectionSize.Y - 11.0f))),
+		FilmBorder,
+		InPainter.bParentEnabled ? ESlateDrawEffect::None : ESlateDrawEffect::DisabledEffect
+	);
+
+	InPainter.DrawElements.PopClip();
+
+	int32 LayerId = FThumbnailSection::OnPaintSection(InPainter);
+
+	// @todo gmp: draw media cache state
+
+	++LayerId;
+
+	return LayerId;
+}
+
+
+/* ICustomThumbnailClient interface
+ *****************************************************************************/
+
+void FMediaThumbnailSection::Draw(FTrackEditorThumbnail& TrackEditorThumbnail)
+{
+	if ((MediaTexture == nullptr) || (MediaTexture->Resource == nullptr) || !MediaTexture->Resource->TextureRHI.IsValid())
+	{
+		return;
+	}
+
+	// get target texture resource
+	FTexture2DRHIRef Texture2DRHI = MediaTexture->Resource->TextureRHI->GetTexture2D();
+
+	if (!Texture2DRHI.IsValid())
+	{
+		return;
+	}
+
+	// seek media player to thumbnail position
+	float VideoTime = FMath::Max(0.0f, TrackEditorThumbnail.GetEvalPosition());
+
+	if (!MediaPlayer->Seek(FTimespan(int64(ETimespan::TicksPerSecond * VideoTime))))
+	{
+		return;
+	}
+
+	// resolve media player texture to track editor thumbnail
+	TrackEditorThumbnail.CopyTextureIn(Texture2DRHI);
+
+	TSharedPtr<ISequencer> Sequencer = SequencerPtr.Pin();
+
+	if (Sequencer.IsValid())
+	{
+		TrackEditorThumbnail.SetupFade(Sequencer->GetSequencerWidget());
+	}
+}
+
+
+void FMediaThumbnailSection::Setup()
+{
+	UMovieSceneMediaSection* MediaSection = CastChecked<UMovieSceneMediaSection>(Section);
+	UMediaSource* MediaSource = MediaSection->GetMediaSource();
+
+	if (!MediaSource)
+	{
+		return;
+	}
+
+	// create internal player
+	if (MediaPlayer == nullptr)
+	{
+		MediaPlayer = NewObject<UMediaPlayer>(GetTransientPackage(), MakeUniqueObjectName(GetTransientPackage(), UMediaPlayer::StaticClass()));
+	}
+
+	// create target texture
+	if (MediaTexture == nullptr)
+	{
+		MediaTexture = NewObject<UMediaTexture>(GetTransientPackage(), MakeUniqueObjectName(GetTransientPackage(), UMediaTexture::StaticClass()));
+		MediaTexture->MediaPlayer = MediaPlayer;
+		MediaTexture->UpdateResource();
+	}
+
+	// open latest media source
+	if (MediaPlayer->GetUrl() != MediaSource->GetUrl())
+	{
+		if (MediaPlayer->CanPlaySource(MediaSource))
+		{
+			MediaPlayer->OpenSource(MediaSource);
+		}
+	}
+
+	MediaPlayer->Pause();
+}
+
+
+#undef LOCTEXT_NAMESPACE
diff --git a/Engine/Plugins/Media/MediaCompositing/Source/MediaCompositingEditor/Private/Sequencer/MediaTrackEditor.cpp b/Engine/Plugins/Media/MediaCompositing/Source/MediaCompositingEditor/Private/Sequencer/MediaTrackEditor.cpp
new file mode 100644
--- /dev/null
+++ b/Engine/Plugins/Media/MediaCompositing/Source/MediaCompositingEditor/Private/Sequencer/MediaTrackEditor.cpp
@@ -0,0 +1,318 @@
+// Copyright 1998-2016 Epic Games, Inc. All Rights Reserved.
+
+#include "MediaTrackEditor.h"
+
+#include "ContentBrowserModule.h"
+#include "EditorStyleSet.h"
+#include "Framework/MultiBox/MultiBoxBuilder.h"
+#include "IContentBrowserSingleton.h"
+#include "ISequencer.h"
+#include "ISequencerObjectChangeListener.h"
+#include "MediaPlayer.h"
+#include "MediaSource.h"
+#include "MovieSceneBinding.h"
+#include "MovieSceneMediaSection.h"
+#include "MovieSceneMediaTrack.h"
+#include "MovieSceneToolsUserSettings.h"
+#include "MediaPlane.h"
+#include "MediaPlaneComponent.h"
+#include "SequencerUtilities.h"
+#include "TrackEditorThumbnail/TrackEditorThumbnailPool.h"
+#include "Widgets/SBoxPanel.h"
+#include "Widgets/Layout/SBox.h"
+
+#include "MediaThumbnailSection.h"
+
+
+#define LOCTEXT_NAMESPACE "FMediaTrackEditor"
+
+
+/* FMediaTrackEditor static functions
+ *****************************************************************************/
+
+TArray<FAnimatedPropertyKey, TInlineAllocator<1>> FMediaTrackEditor::GetAnimatedPropertyTypes()
+{
+	return TArray<FAnimatedPropertyKey, TInlineAllocator<1>>({ FAnimatedPropertyKey::FromObjectType(UMediaPlayer::StaticClass()) });
+}
+
+
+/* FMediaTrackEditor structors
+ *****************************************************************************/
+
+FMediaTrackEditor::FMediaTrackEditor(TSharedRef<ISequencer> InSequencer)
+	: FMovieSceneTrackEditor(InSequencer)
+	, PropertyKey(FAnimatedPropertyKey::FromObjectType(UMediaPlayer::StaticClass()))
+{
+	ThumbnailPool = MakeShared<FTrackEditorThumbnailPool>(InSequencer);
+	OnPropertyChangedHandle = InSequencer->GetObjectChangeListener().GetOnAnimatablePropertyChanged(PropertyKey).AddRaw(this, &FMediaTrackEditor::OnAnimatedPropertyChanged);
+}
+
+
+FMediaTrackEditor::~FMediaTrackEditor()
+{
+	TSharedPtr<ISequencer> SequencerPtr = GetSequencer();
+
+	if (SequencerPtr.IsValid())
+	{
+		SequencerPtr->GetObjectChangeListener().GetOnAnimatablePropertyChanged(PropertyKey).Remove(OnPropertyChangedHandle);
+	}
+}
+
+
+/* FMovieSceneTrackEditor interface
+ *****************************************************************************/
+
+UMovieSceneTrack* FMediaTrackEditor::AddTrack(UMovieScene* FocusedMovieScene, const FGuid& ObjectHandle, TSubclassOf<class UMovieSceneTrack> TrackClass, FName UniqueTypeName)
+{
+	UMovieSceneTrack* Track = FocusedMovieScene->AddTrack(TrackClass, ObjectHandle);
+	if (UMovieSceneMediaTrack* MediaTrack = Cast<UMovieSceneMediaTrack>(Track))
+	{
+		MediaTrack->UniqueTrackName = UniqueTypeName;
+	}
+	return Track;
+}
+
+void FMediaTrackEditor::OnAnimatedPropertyChanged(const FPropertyChangedParams& PropertyChangedParams)
+{
+	const UProperty* ChangedProperty = PropertyChangedParams.PropertyPath.GetLeafMostProperty().Property.Get();
+	if (!ChangedProperty)
+	{
+		return;
+	}
+
+	FText DisplayText = ChangedProperty->GetDisplayNameText();
+	FName UniqueName(*PropertyChangedParams.PropertyPath.ToString(TEXT(".")));
+
+	// Set up the appropriate name for the track from an array index if necessary
+	for (int32 PropertyIndex = PropertyChangedParams.PropertyPath.GetNumProperties() - 1; PropertyIndex >= 0; --PropertyIndex)
+	{
+		const FPropertyInfo& Info = PropertyChangedParams.PropertyPath.GetPropertyInfo(PropertyIndex);
+		const UArrayProperty* ParentArrayProperty = PropertyIndex > 0 ? Cast<UArrayProperty>(PropertyChangedParams.PropertyPath.GetPropertyInfo(PropertyIndex-1).Property.Get()) : nullptr;
+
+		UProperty* Property = Info.Property.Get();
+		if (Property && Info.ArrayIndex != INDEX_NONE)
+		{
+			DisplayText = FText::Format(LOCTEXT("MediaTrackNameFormat", "{0} ({1}[{2}])"),
+				ChangedProperty->GetDisplayNameText(),
+				(ParentArrayProperty ? ParentArrayProperty : Property)->GetDisplayNameText(),
+				FText::AsNumber(Info.ArrayIndex)
+				);
+			break;
+		}
+	}
+
+	for (UObject* Object : PropertyChangedParams.ObjectsThatChanged)
+	{
+		FFindOrCreateHandleResult HandleResult = FindOrCreateHandleToObject(Object);
+		if (!ensure(HandleResult.Handle.IsValid()))
+		{
+			continue;
+		}
+		
+		FFindOrCreateTrackResult TrackResult = FindOrCreateTrackForObject(HandleResult.Handle, UMovieSceneMediaTrack::StaticClass(), UniqueName);
+		UMovieSceneMediaTrack* MediaTrack = Cast<UMovieSceneMediaTrack>(TrackResult.Track);
+		
+		if (TrackResult.bWasCreated)
+		{
+			MediaTrack->SetPropertyNameAndPath(ChangedProperty->GetFName(), PropertyChangedParams.PropertyPath.ToString(TEXT(".")));
+			MediaTrack->SetDisplayName(DisplayText);
+		}
+
+		GetSequencer()->NotifyMovieSceneDataChanged(EMovieSceneDataChangeType::MovieSceneStructureItemAdded);
+	}
+}
+
+void FMediaTrackEditor::BuildObjectBindingTrackMenu(FMenuBuilder& MenuBuilder, const FGuid& ObjectBinding, const UClass* ObjectClass)
+{
+	// We only know how to add specific properties for image planes and components. Anything else must be keyed through the generic MediaPlayer property
+	if (!ObjectClass->IsChildOf(AMediaPlane::StaticClass()) && !ObjectClass->IsChildOf(UMediaPlaneComponent::StaticClass()))
+	{
+		return;
+	}
+
+	// Find the spawned object or its template
+	UObject* Object = GetSequencer()->FindSpawnedObjectOrTemplate(ObjectBinding);
+	if (!Object)
+	{
+		return;
+	}
+
+	auto GeneratePropertyPath = [this](UMediaPlaneComponent* MediaPlaneComponent)
+	{
+		check(MediaPlaneComponent);
+
+		UStructProperty* MediaPlanesProperty = MediaPlaneComponent->GetMediaPlaneProperty();
+
+		check(MediaPlanesProperty);
+
+		TSharedRef<FPropertyPath> Path = FPropertyPath::CreateEmpty();
+		Path->AddProperty(FPropertyInfo(MediaPlanesProperty));
+
+		return Path;
+	};
+
+	auto AddNewTrack = [this](UMediaPlaneComponent* MediaPlaneComponent, TSharedRef<FPropertyPath> PropertyPath)
+	{
+		check(MediaPlaneComponent);
+
+		FPropertyChangedParams ChangedParams(TArray<UObject*>({ MediaPlaneComponent }), *PropertyPath, NAME_None, ESequencerKeyMode::ManualKeyForced);
+		this->OnAnimatedPropertyChanged(ChangedParams);
+	};
+
+	// Try and root out an image plane component
+	UMediaPlaneComponent* Component = Cast<UMediaPlaneComponent>(Object);
+	if (AMediaPlane* MediaPlane = Cast<AMediaPlane>(Object))
+	{
+		Component = Cast<UMediaPlaneComponent>(MediaPlane->GetRootComponent());
+	}
+
+	if (Component)
+	{
+		UMovieScene* MovieScene = GetSequencer()->GetFocusedMovieSceneSequence()->GetMovieScene();
+
+		TSharedRef<FPropertyPath> Path = GeneratePropertyPath(Component);
+
+		bool bCanAddTrack = true;
+
+		const FMovieSceneBinding* Binding = MovieScene->GetBindings().FindByPredicate([=](const FMovieSceneBinding& InBinding){ return InBinding.GetObjectGuid() == ObjectBinding; });
+		if (Binding)
+		{
+			FString PredicatePath = Path->ToString(TEXT("."));
+
+			bCanAddTrack = !Binding->GetTracks().ContainsByPredicate(
+				[&](UMovieSceneTrack* Track)
+				{
+					UMovieSceneMediaTrack* MediaTrack = Cast<UMovieSceneMediaTrack>(Track);
+					return MediaTrack && MediaTrack->GetPropertyPath() == PredicatePath;
+				}
+			);
+		}
+
+		MenuBuilder.AddMenuEntry(
+			LOCTEXT("AddMediaTrack_Text", "Media"),
+			LOCTEXT("AddMediaTrack_Tip", "Adds a media track that controls the media presented to the media plane."),
+			FSlateIcon(),
+			FUIAction(
+				FExecuteAction::CreateLambda(AddNewTrack, Component, Path),
+				FCanExecuteAction::CreateLambda([=]{ return bCanAddTrack; })
+			)
+		);
+	}
+}
+
+TSharedPtr<SWidget> FMediaTrackEditor::BuildOutlinerEditWidget(const FGuid& ObjectBinding, UMovieSceneTrack* Track, const FBuildEditWidgetParams& Params)
+{
+	UMovieSceneMediaTrack* MediaTrack = Cast<UMovieSceneMediaTrack>(Track);
+	if (!MediaTrack)
+	{
+		return SNullWidget::NullWidget;
+	}
+
+	auto CreatePicker = [this, MediaTrack]
+	{
+		FAssetPickerConfig AssetPickerConfig;
+		AssetPickerConfig.OnAssetSelected = FOnAssetSelected::CreateRaw(this, &FMediaTrackEditor::AddNewSection, MediaTrack);
+		AssetPickerConfig.bAllowNullSelection = false;
+		AssetPickerConfig.InitialAssetViewType = EAssetViewType::List;
+
+		AssetPickerConfig.Filter.bRecursiveClasses = true;
+		AssetPickerConfig.Filter.ClassNames.Add(UMediaSource::StaticClass()->GetFName());
+
+		FContentBrowserModule& ContentBrowserModule = FModuleManager::Get().LoadModuleChecked<FContentBrowserModule>(TEXT("ContentBrowser"));
+
+		TSharedRef<SBox> Picker = SNew(SBox)
+			.WidthOverride(300.0f)
+			.HeightOverride(300.f)
+			[
+				ContentBrowserModule.Get().CreateAssetPicker(AssetPickerConfig)
+			];
+
+		return Picker;
+	};
+
+	return SNew(SHorizontalBox)
+		+ SHorizontalBox::Slot()
+		.AutoWidth()
+		.VAlign(VAlign_Center)
+		[
+			FSequencerUtilities::MakeAddButton(LOCTEXT("AddMediaSection_Text", "Media"), FOnGetContent::CreateLambda(CreatePicker), Params.NodeIsHovered)
+		];
+}
+
+void FMediaTrackEditor::AddNewSection(const FAssetData& AssetData, UMovieSceneMediaTrack* Track)
+{
+	UMediaSource* MediaSource = Cast<UMediaSource>(AssetData.GetAsset());
+	if (!MediaSource)
+	{
+		return;
+	}
+
+	check(Track);
+
+	float TimeToStart = GetSequencer()->GetLocalTime();
+
+	UMediaPlayer* TransientPlayer = NewObject<UMediaPlayer>(GetTransientPackage(), MakeUniqueObjectName(GetTransientPackage(), UMediaPlayer::StaticClass()));
+	TransientPlayer->AddToRoot();
+
+	TRange<float> SectionRange(TimeToStart, TimeToStart + 1.f);
+
+	if (TransientPlayer->OpenSource(MediaSource))
+	{
+		// @todo: Make this more robust for all media types
+		FTimespan Length = TransientPlayer->GetDuration();
+		if (Length.GetTotalSeconds() > 0)
+		{
+			SectionRange = TRange<float>(TimeToStart, TRangeBound<float>::Inclusive(TimeToStart + float(Length.GetTicks()) / ETimespan::TicksPerSecond));
+		}
+
+		TransientPlayer->Close();
+	}
+
+	TransientPlayer->RemoveFromRoot();
+
+	// Find a row to put this video on
+	TArray<int32> InvalidRows;
+	for (UMovieSceneSection* Section : Track->GetAllSections())
+	{
+		if (Section->GetRange().Overlaps(SectionRange))
+		{
+			InvalidRows.Add(Section->GetRowIndex());
+		}
+	}
+
+	int32 BestRow = 0;
+	for ( ; ; ++BestRow)
+	{
+		if (!InvalidRows.Contains(BestRow))
+		{
+			break;
+		}
+	}
+
+	UMovieSceneMediaSection* Section = CastChecked<UMovieSceneMediaSection>(Track->CreateNewSection());
+	Section->SetRange(SectionRange);
+	Section->SetRowIndex(BestRow);
+	Section->SetMediaSource(MediaSource);
+	Track->AddSection(*Section);
+
+	GetSequencer()->NotifyMovieSceneDataChanged(EMovieSceneDataChangeType::MovieSceneStructureItemAdded);
+}
+
+void FMediaTrackEditor::Tick(float DeltaTime)
+{
+	ThumbnailPool->DrawThumbnails();
+}
+
+TSharedRef<ISequencerSection> FMediaTrackEditor::MakeSectionInterface(UMovieSceneSection& SectionObject, UMovieSceneTrack& Track, FGuid ObjectBinding)
+{
+	check(SupportsType(SectionObject.GetOuter()->GetClass()));
+	return MakeShared<FMediaThumbnailSection>(*CastChecked<UMovieSceneMediaSection>(&SectionObject), ThumbnailPool, GetSequencer());
+}
+
+bool FMediaTrackEditor::SupportsType(TSubclassOf<UMovieSceneTrack> TrackClass) const
+{
+	return TrackClass.Get() && TrackClass.Get()->IsChildOf(UMovieSceneMediaTrack::StaticClass());
+}
+
+
+#undef LOCTEXT_NAMESPACE
diff --git a/Engine/Plugins/Media/MediaCompositing/Source/MediaCompositingEditor/Private/Sequencer/MediaTrackEditor.h b/Engine/Plugins/Media/MediaCompositing/Source/MediaCompositingEditor/Private/Sequencer/MediaTrackEditor.h
new file mode 100644
--- /dev/null
+++ b/Engine/Plugins/Media/MediaCompositing/Source/MediaCompositingEditor/Private/Sequencer/MediaTrackEditor.h
@@ -0,0 +1,76 @@
+// Copyright 1998-2016 Epic Games, Inc. All Rights Reserved.
+
+#pragma once
+
+#include "AnimatedPropertyKey.h"
+#include "AssetData.h"
+#include "MovieSceneTrackEditor.h"
+#include "Templates/SharedPointer.h"
+
+class ISequencer;
+class FTrackEditorThumbnailPool;
+class UMediaSource;
+class UMovieSceneMediaTrack;
+
+
+/**
+ * Track editor that understands how to animate MediaPlayer properties on objects
+ */
+class FMediaTrackEditor
+	: public FMovieSceneTrackEditor
+{
+public:
+
+	/**
+	 * Create a new media track editor instance.
+	 *
+	 * @param OwningSequencer The sequencer object that will own the track editor.
+	 * @return The new track editor.
+	 */
+	static TSharedRef<ISequencerTrackEditor> CreateTrackEditor(TSharedRef<ISequencer> OwningSequencer)
+	{
+		return MakeShared<FMediaTrackEditor>(OwningSequencer);
+	}
+
+	/**
+	 * Get the list of all property types that this track editor animates.
+	 *
+	 * @return List of animated properties.
+	 */
+	static TArray<FAnimatedPropertyKey, TInlineAllocator<1>> GetAnimatedPropertyTypes();
+
+public:
+
+	/**
+	 * Create and initialize a new instance.
+	 *
+	 * @param InSequencer The sequencer object that owns this track editor.
+	 */
+	FMediaTrackEditor(TSharedRef<ISequencer> InSequencer);
+
+	/** Virtual destructor. */
+	virtual ~FMediaTrackEditor();
+
+public:
+
+	//~ FMovieSceneTrackEditor interface
+
+	virtual TSharedRef<ISequencerSection> MakeSectionInterface(UMovieSceneSection& SectionObject, UMovieSceneTrack& Track, FGuid ObjectBinding) override;
+	virtual bool SupportsType(TSubclassOf<UMovieSceneTrack> TrackClass) const override;
+	virtual void Tick(float DeltaTime) override;
+	virtual void BuildObjectBindingTrackMenu(FMenuBuilder& MenuBuilder, const FGuid& ObjectBinding, const UClass* ObjectClass) override;
+	virtual TSharedPtr<SWidget> BuildOutlinerEditWidget(const FGuid& ObjectBinding, UMovieSceneTrack* Track, const FBuildEditWidgetParams& Params) override;
+	virtual UMovieSceneTrack* AddTrack(UMovieScene* FocusedMovieScene, const FGuid& ObjectHandle, TSubclassOf<class UMovieSceneTrack> TrackClass, FName UniqueTypeName) override;
+
+protected:
+
+	void OnAnimatedPropertyChanged(const FPropertyChangedParams& PropertyChangedParams);
+
+	void AddNewSection(const FAssetData& Asset, UMovieSceneMediaTrack* Track);
+
+private:
+
+	FAnimatedPropertyKey PropertyKey;
+	FDelegateHandle OnPropertyChangedHandle;
+	TSharedPtr<FTrackEditorThumbnailPool> ThumbnailPool;
+};
diff --git a/Engine/Plugins/Media/MediaPlayerEditor/Source/MediaPlayerEditor/MediaPlayerEditor.Build.cs b/Engine/Plugins/Media/MediaPlayerEditor/Source/MediaPlayerEditor/MediaPlayerEditor.Build.cs
--- a/Engine/Plugins/Media/MediaPlayerEditor/Source/MediaPlayerEditor/MediaPlayerEditor.Build.cs
+++ b/Engine/Plugins/Media/MediaPlayerEditor/Source/MediaPlayerEditor/MediaPlayerEditor.Build.cs
@@ -12,35 +12,34 @@ public class MediaPlayerEditor : ModuleRules
 				"MainFrame",
 				"Media",
 				"WorkspaceMenuStructure",
-			}
-		);
+			});
 
 		PrivateIncludePaths.AddRange(
 			new string[] {
 				"MediaPlayerEditor/Private",
 				"MediaPlayerEditor/Private/AssetTools",
-                "MediaPlayerEditor/Private/Customizations",
+				"MediaPlayerEditor/Private/Customizations",
 				"MediaPlayerEditor/Private/Factories",
 				"MediaPlayerEditor/Private/Models",
 				"MediaPlayerEditor/Private/Shared",
-                "MediaPlayerEditor/Private/Toolkits",
-                "MediaPlayerEditor/Private/Widgets",
-			}
-		);
+				"MediaPlayerEditor/Private/Toolkits",
+				"MediaPlayerEditor/Private/Widgets",
+			});
 
 		PrivateDependencyModuleNames.AddRange(
 			new string[] {
+				"AudioMixer",
 				"ContentBrowser",
 				"Core",
 				"CoreUObject",
-                "DesktopPlatform",
-                "DesktopWidgets",
+				"DesktopPlatform",
+				"DesktopWidgets",
 				"EditorStyle",
-                "EditorWidgets",
+				"EditorWidgets",
 				"Engine",
 				"InputCore",
 				"MediaAssets",
-                "MediaUtils",
+				"MediaUtils",
 				"PropertyEditor",
 				"RenderCore",
 				"RHI",
@@ -49,8 +48,7 @@ public class MediaPlayerEditor : ModuleRules
 				"SlateCore",
 				"TextureEditor",
 				"UnrealEd",
-			}
-		);
+			});
 
 		PrivateIncludePathModuleNames.AddRange(
 			new string[] {
@@ -58,7 +56,6 @@ public class MediaPlayerEditor : ModuleRules
 				"Media",
 				"UnrealEd",
 				"WorkspaceMenuStructure",
-			}
-		);
+			});
 	}
 }
diff --git a/Engine/Plugins/Media/MediaPlayerEditor/Source/MediaPlayerEditor/Private/AssetTools/MediaPlayerActions.cpp b/Engine/Plugins/Media/MediaPlayerEditor/Source/MediaPlayerEditor/Private/AssetTools/MediaPlayerActions.cpp
--- a/Engine/Plugins/Media/MediaPlayerEditor/Source/MediaPlayerEditor/Private/AssetTools/MediaPlayerActions.cpp
+++ b/Engine/Plugins/Media/MediaPlayerEditor/Source/MediaPlayerEditor/Private/AssetTools/MediaPlayerActions.cpp
@@ -7,10 +7,8 @@
 #include "IContentBrowserSingleton.h"
 #include "ContentBrowserModule.h"
 #include "EditorStyleSet.h"
-#include "MediaSoundWave.h"
 #include "MediaTexture.h"
 #include "Toolkits/MediaPlayerEditorToolkit.h"
-#include "Factories/MediaSoundWaveFactoryNew.h"
 #include "Factories/MediaTextureFactoryNew.h"
 
 
@@ -34,34 +32,6 @@ bool FMediaPlayerActions::CanFilter()
 }
 
 
-void FMediaPlayerActions::GetActions(const TArray<UObject*>& InObjects, FMenuBuilder& MenuBuilder)
-{
-	FAssetTypeActions_Base::GetActions(InObjects, MenuBuilder);
-
-	auto MediaPlayers = GetTypedWeakObjectPtrs<UMediaPlayer>(InObjects);
-
-	MenuBuilder.AddMenuEntry(
-		LOCTEXT("MediaPlayer_CreateMediaSoundWave", "Create Media Sound Wave"),
-		LOCTEXT("MediaPlayer_CreateMediaSoundWaveTooltip", "Creates a new MediaSoundWave and assigns it to this MediaPlayer asset."),
-		FSlateIcon(FEditorStyle::GetStyleSetName(), "ClassIcon.MediaSoundWave"),
-		FUIAction(
-			FExecuteAction::CreateSP(this, &FMediaPlayerActions::ExecuteCreateMediaSoundWave, MediaPlayers),
-			FCanExecuteAction()
-		)
-	);
-
-	MenuBuilder.AddMenuEntry(
-		LOCTEXT("MediaPlayer_CreateMediaTexture", "Create Media Texture"),
-		LOCTEXT("MediaPlayer_CreateMediaTextureTooltip", "Creates a new MediaTexture and assigns it to this MediaPlayer asset."),
-		FSlateIcon(FEditorStyle::GetStyleSetName(), "ClassIcon.MediaTexture"),
-		FUIAction(
-			FExecuteAction::CreateSP(this, &FMediaPlayerActions::ExecuteCreateMediaTexture, MediaPlayers),
-			FCanExecuteAction()
-		)
-	);
-}
-
-
 uint32 FMediaPlayerActions::GetCategories()
 {
 	return EAssetTypeCategories::Media;
@@ -86,12 +56,6 @@ FColor FMediaPlayerActions::GetTypeColor() const
 }
 
 
-bool FMediaPlayerActions::HasActions(const TArray<UObject*>& InObjects) const
-{
-	return true;
-}
-
-
 void FMediaPlayerActions::OpenAssetEditor(const TArray<UObject*>& InObjects, TSharedPtr<IToolkitHost> EditWithinLevelEditor)
 {
 	EToolkitMode::Type Mode = EditWithinLevelEditor.IsValid()
@@ -111,93 +75,4 @@ void FMediaPlayerActions::OpenAssetEditor(const TArray<UObject*>& InObjects, TSh
 }
 
 
-/* FAssetTypeActions_MediaPlayer callbacks
- *****************************************************************************/
-
-void FMediaPlayerActions::ExecuteCreateMediaSoundWave(TArray<TWeakObjectPtr<UMediaPlayer>> Objects)
-{
-	const FString DefaultSuffix = TEXT("_Sound");
-
-	IContentBrowserSingleton& ContentBrowserSingleton = FModuleManager::LoadModuleChecked<FContentBrowserModule>("ContentBrowser").Get();
-	TArray<UObject*> ObjectsToSync;
-
-	for (auto ObjIt = Objects.CreateConstIterator(); ObjIt; ++ObjIt)
-	{
-		auto MediaPlayer = (*ObjIt).Get();
-
-		if (MediaPlayer == nullptr)
-		{
-			continue;
-		}
-
-		// generate unique name
-		FString Name;
-		FString PackageName;
-
-		CreateUniqueAssetName(MediaPlayer->GetOutermost()->GetName(), DefaultSuffix, PackageName, Name);
-
-		// create asset
-		auto Factory = NewObject<UMediaSoundWaveFactoryNew>();
-
-		FAssetToolsModule& AssetToolsModule = FModuleManager::GetModuleChecked<FAssetToolsModule>("AssetTools");
-		UObject* NewAsset = AssetToolsModule.Get().CreateAsset(Name, FPackageName::GetLongPackagePath(PackageName), UMediaSoundWave::StaticClass(), Factory);
-
-		if (NewAsset != nullptr)
-		{
-			MediaPlayer->SetSoundWave(Cast<UMediaSoundWave>(NewAsset));
-			MediaPlayer->MarkPackageDirty();
-			ObjectsToSync.Add(NewAsset);
-		}
-	}
-
-	if (ObjectsToSync.Num() > 0)
-	{
-		ContentBrowserSingleton.SyncBrowserToAssets(ObjectsToSync);
-	}
-}
-
-
-void FMediaPlayerActions::ExecuteCreateMediaTexture(TArray<TWeakObjectPtr<UMediaPlayer>> Objects)
-{
-	const FString DefaultSuffix = TEXT("_Video");
-
-	IContentBrowserSingleton& ContentBrowserSingleton = FModuleManager::LoadModuleChecked<FContentBrowserModule>("ContentBrowser").Get();
-	TArray<UObject*> ObjectsToSync;
-
-	for (auto ObjIt = Objects.CreateConstIterator(); ObjIt; ++ObjIt)
-	{
-		auto MediaPlayer = (*ObjIt).Get();
-
-		if (MediaPlayer == nullptr)
-		{
-			continue;
-		}
-
-		// generate unique name
-		FString Name;
-		FString PackageName;
-
-		CreateUniqueAssetName(MediaPlayer->GetOutermost()->GetName(), DefaultSuffix, PackageName, Name);
-
-		// create asset
-		auto Factory = NewObject<UMediaTextureFactoryNew>();
-
-		FAssetToolsModule& AssetToolsModule = FModuleManager::GetModuleChecked<FAssetToolsModule>("AssetTools");
-		UObject* NewAsset = AssetToolsModule.Get().CreateAsset(Name, FPackageName::GetLongPackagePath(PackageName), UMediaTexture::StaticClass(), Factory);
-
-		if (NewAsset != nullptr)
-		{
-			MediaPlayer->SetVideoTexture(Cast<UMediaTexture>(NewAsset));
-			MediaPlayer->MarkPackageDirty();
-			ObjectsToSync.Add(NewAsset);
-		}
-	}
-
-	if (ObjectsToSync.Num() > 0)
-	{
-		ContentBrowserSingleton.SyncBrowserToAssets(ObjectsToSync);
-	}
-}
-
-
 #undef LOCTEXT_NAMESPACE
diff --git a/Engine/Plugins/Media/MediaPlayerEditor/Source/MediaPlayerEditor/Private/AssetTools/MediaPlayerActions.h b/Engine/Plugins/Media/MediaPlayerEditor/Source/MediaPlayerEditor/Private/AssetTools/MediaPlayerActions.h
--- a/Engine/Plugins/Media/MediaPlayerEditor/Source/MediaPlayerEditor/Private/AssetTools/MediaPlayerActions.h
+++ b/Engine/Plugins/Media/MediaPlayerEditor/Source/MediaPlayerEditor/Private/AssetTools/MediaPlayerActions.h
@@ -2,13 +2,14 @@
 
 #pragma once
 
-#include "CoreMinimal.h"
-#include "Styling/ISlateStyle.h"
 #include "AssetTypeActions_Base.h"
+#include "Styling/ISlateStyle.h"
+#include "Templates/SharedPointer.h"
 
 class FMenuBuilder;
 class UMediaPlayer;
 
+
 /**
  * Implements an action for UMediaPlayer assets.
  */
@@ -29,22 +30,12 @@ public:
 	//~ FAssetTypeActions_Base interface
 
 	virtual bool CanFilter() override;
-	virtual void GetActions(const TArray<UObject*>& InObjects, FMenuBuilder& MenuBuilder) override;
 	virtual uint32 GetCategories() override;
 	virtual FText GetName() const override;
 	virtual UClass* GetSupportedClass() const override;
 	virtual FColor GetTypeColor() const override;
-	virtual bool HasActions(const TArray<UObject*>& InObjects) const override;
 	virtual void OpenAssetEditor(const TArray<UObject*>& InObjects, TSharedPtr<IToolkitHost> EditWithinLevelEditor = TSharedPtr<IToolkitHost>()) override;
 
-private:
-
-	/** Callback for selecting CreateMediaSoundWave. */
-	void ExecuteCreateMediaSoundWave(TArray<TWeakObjectPtr<UMediaPlayer>> Objects);
-
-	/** Callback for selecting CreateMediaTexture. */
-	void ExecuteCreateMediaTexture(TArray<TWeakObjectPtr<UMediaPlayer>> Objects);	
-
 private:
 
 	/** Pointer to the style set to use for toolkits. */
diff --git a/Engine/Plugins/Media/MediaPlayerEditor/Source/MediaPlayerEditor/Private/Customizations/MediaTextureCustomization.cpp b/Engine/Plugins/Media/MediaPlayerEditor/Source/MediaPlayerEditor/Private/Customizations/MediaTextureCustomization.cpp
--- a/Engine/Plugins/Media/MediaPlayerEditor/Source/MediaPlayerEditor/Private/Customizations/MediaTextureCustomization.cpp
+++ b/Engine/Plugins/Media/MediaPlayerEditor/Source/MediaPlayerEditor/Private/Customizations/MediaTextureCustomization.cpp
@@ -1,154 +1,27 @@
 // Copyright 1998-2017 Epic Games, Inc. All Rights Reserved.
 
 #include "Customizations/MediaTextureCustomization.h"
+
+#include "DetailCategoryBuilder.h"
+#include "DetailLayoutBuilder.h"
+#include "Engine/Texture.h"
 #include "Widgets/SNullWidget.h"
 #include "PropertyHandle.h"
 
 
-#define LOCTEXT_NAMESPACE "FMediaTextureCustomization"
-
-
 /* IDetailCustomization interface
  *****************************************************************************/
 
 void FMediaTextureCustomization::CustomizeDetails(IDetailLayoutBuilder& DetailBuilder)
-{/*
-	DetailBuilder.GetObjectsBeingCustomized(CustomizedMediaTextures);
-	MediaPlayerProperty = DetailBuilder.GetProperty("MediaPlayer");
-	VideoTrackIndexProperty = DetailBuilder.GetProperty("VideoTrackIndex");
-
-	// customize MediaPlayer category
-	IDetailCategoryBuilder& MediaPlayerCategory = DetailBuilder.EditCategory(TEXT("MediaPlayer"));
-	{
-		// video track index
-		IDetailPropertyRow& VideoTrackIndexRow = MediaPlayerCategory.AddProperty(VideoTrackIndexProperty);
-
-		VideoTrackIndexRow.DisplayName(LOCTEXT("VideoTrack", "Video Track"));
-		VideoTrackIndexRow.CustomWidget()
-			.NameContent()
-			[
-				VideoTrackIndexProperty->CreatePropertyNameWidget()
-			]
-			.ValueContent()
-			.MaxDesiredWidth(0.0f)
-			[
-				SNew(SComboButton)
-					.OnGetMenuContent(this, &FMediaTextureCustomization::HandleVideoTrackComboButtonMenuContent)
- 					.ContentPadding(FMargin( 2.0f, 2.0f ))
-					.ButtonContent()
-					[
-						SNew(STextBlock) 
-							.Text(this, &FMediaTextureCustomization::HandleVideoTrackComboButtonText)
-							.Font(IDetailLayoutBuilder::GetDetailFont())
-					]
-			];
-	}*/
-}
-
-
-/* FMediaTextureCustomization callbacks
- *****************************************************************************/
-
-TSharedRef<SWidget> FMediaTextureCustomization::HandleVideoTrackComboButtonMenuContent() const
-{/*
-	// get assigned media player asset
-	UObject* MediaPlayerObj = nullptr;
-	FPropertyAccess::Result Result = MediaPlayerProperty->GetValue(MediaPlayerObj);
-
-	if (Result == FPropertyAccess::MultipleValues)
-	{
-		return SNullWidget::NullWidget;
-	}
-
-	UMediaPlayer* MediaPlayer = Cast<UMediaPlayer>(MediaPlayerObj);
-
-	if (MediaPlayer == nullptr)
-	{
-		return SNullWidget::NullWidget;
-	}
-
-	// get media tracks
-	TSharedPtr<IMediaPlayer, ESPMode::ThreadSafe> Player = MediaPlayer->GetPlayer();
-
-	if (!Player.IsValid())
-	{
-		return SNullWidget::NullWidget;
-	}
-
-	const TArray<IMediaVideoTrackRef>& VideoTracks = Player->GetVideoTracks();
-
-	// populate the menu
-	FMenuBuilder MenuBuilder(true, nullptr);
-
-	for (int32 VideoTrackIndex = 0; VideoTrackIndex < VideoTracks.Num(); ++VideoTrackIndex)
-	{
-		FUIAction Action(FExecuteAction::CreateSP(this, &FMediaTextureCustomization::HandleVideoTrackComboButtonMenuEntryExecute, VideoTrackIndex));
-
-		MenuBuilder.AddMenuEntry(
-			VideoTracks[VideoTrackIndex]->GetStream().GetDisplayName(),
-			TAttribute<FText>(),
-			FSlateIcon(),
-			Action
-		);
-	}
-
-	return MenuBuilder.MakeWidget();*/
-	return SNullWidget::NullWidget;
-}
-
-
-void FMediaTextureCustomization::HandleVideoTrackComboButtonMenuEntryExecute(int32 TrackIndex)
 {
-	VideoTrackIndexProperty->SetValue(TrackIndex);
-}
-
-
-FText FMediaTextureCustomization::HandleVideoTrackComboButtonText() const
-{/*
-	// get assigned media player asset
-	UObject* MediaPlayerObj = nullptr;
-	FPropertyAccess::Result Result = MediaPlayerProperty->GetValue(MediaPlayerObj);
-
-	if (Result != FPropertyAccess::Success)
+	// customize 'Texture' category
+	IDetailCategoryBuilder& MediaTextureCategory = DetailBuilder.EditCategory("MediaTexture");
 	{
-		return FText::GetEmpty();
+		MediaTextureCategory.AddProperty(GET_MEMBER_NAME_CHECKED(UTexture, CompressionSettings), UTexture::StaticClass());
+		MediaTextureCategory.AddProperty(GET_MEMBER_NAME_CHECKED(UTexture, Filter), UTexture::StaticClass());
+		MediaTextureCategory.AddProperty(GET_MEMBER_NAME_CHECKED(UTexture, SRGB), UTexture::StaticClass());
 	}
 
-	UMediaPlayer* MediaPlayer = Cast<UMediaPlayer>(MediaPlayerObj);
-
-	if (MediaPlayer == nullptr)
-	{
-		return LOCTEXT("NoMediaPlayerSelected", "No Media Asset selected");
-	}
-
-	// get selected value
-	int32 VideoTrackIndex = -1;
-	Result = VideoTrackIndexProperty->GetValue(VideoTrackIndex);
-
-	if (Result != FPropertyAccess::Success)
-	{
-		return FText::GetEmpty();
-	}
-
-	// get selected media track
-	TSharedPtr<IMediaPlayer, ESPMode::ThreadSafe> Player = MediaPlayer->GetPlayer();
-
-	if (!Player.IsValid())
-	{
-		return LOCTEXT("NoMediaLoaded", "No media loaded");
-	}
-
-	auto VideoTracks = Player->GetVideoTracks();
-
-	if (!VideoTracks.IsValidIndex(VideoTrackIndex))
-	{
-		return LOCTEXT("SelectVideoTrack", "Select a Video Track...");
-	}
-
-	// get track name
-	return VideoTracks[VideoTrackIndex]->GetStream().GetDisplayName();*/
-	return FText::GetEmpty();
+	DetailBuilder.HideCategory("Compression");
+	DetailBuilder.HideCategory("Texture");
 }
-
-
-#undef LOCTEXT_NAMESPACE
diff --git a/Engine/Plugins/Media/MediaPlayerEditor/Source/MediaPlayerEditor/Private/Customizations/MediaTextureCustomization.h b/Engine/Plugins/Media/MediaPlayerEditor/Source/MediaPlayerEditor/Private/Customizations/MediaTextureCustomization.h
--- a/Engine/Plugins/Media/MediaPlayerEditor/Source/MediaPlayerEditor/Private/Customizations/MediaTextureCustomization.h
+++ b/Engine/Plugins/Media/MediaPlayerEditor/Source/MediaPlayerEditor/Private/Customizations/MediaTextureCustomization.h
@@ -34,26 +34,4 @@ public:
 	{
 		return MakeShareable(new FMediaTextureCustomization());
 	}
-
-private:
-
-	/** Callback for generating the menu content of the VideoTrack combo box. */
-	TSharedRef<SWidget> HandleVideoTrackComboButtonMenuContent() const;
-
-	/** Callback for selecting a track in the VideoTrack combo box. */
-	void HandleVideoTrackComboButtonMenuEntryExecute(int32 TrackIndex);
-
-	/** Callback for getting the text of the VideoTrack combo box. */
-	FText HandleVideoTrackComboButtonText() const;
-
-private:
-
-	/** The collection of media textures being customized */
-	TArray<TWeakObjectPtr<UObject>> CustomizedMediaTextures;
-
-	/** Pointer to the MediaPlayer property handle. */
-	TSharedPtr<IPropertyHandle> MediaPlayerProperty;
-
-	/** Pointer to the VideoTrackIndex property handle. */
-	TSharedPtr<IPropertyHandle> VideoTrackIndexProperty;
 };
diff --git a/Engine/Plugins/Media/MediaPlayerEditor/Source/MediaPlayerEditor/Private/Customizations/PlatformMediaSourceCustomization.cpp b/Engine/Plugins/Media/MediaPlayerEditor/Source/MediaPlayerEditor/Private/Customizations/PlatformMediaSourceCustomization.cpp
--- a/Engine/Plugins/Media/MediaPlayerEditor/Source/MediaPlayerEditor/Private/Customizations/PlatformMediaSourceCustomization.cpp
+++ b/Engine/Plugins/Media/MediaPlayerEditor/Source/MediaPlayerEditor/Private/Customizations/PlatformMediaSourceCustomization.cpp
@@ -118,9 +118,9 @@ TSharedRef<SWidget> FPlatformMediaSourceCustomization::MakePlatformMediaSourcesV
 				SNew(SObjectPropertyEntryBox)
 					.AllowedClass(UMediaSource::StaticClass())
 					.AllowClear(true)
-					.OnShouldFilterAsset(this, &FPlatformMediaSourceCustomization::HandleShouldFilterAsset)
-					.ObjectPath(this, &FPlatformMediaSourceCustomization::HandleMediaSourcePropertyEntryObjectPath, Platform->IniPlatformName)
-					.OnObjectChanged(this, &FPlatformMediaSourceCustomization::HandleMediaSourcePropertyEntryBoxChanged, Platform->IniPlatformName)
+					.ObjectPath(this, &FPlatformMediaSourceCustomization::HandleMediaSourceEntryBoxObjectPath, Platform->IniPlatformName)
+					.OnObjectChanged(this, &FPlatformMediaSourceCustomization::HandleMediaSourceEntryBoxChanged, Platform->IniPlatformName)
+					.OnShouldFilterAsset(this, &FPlatformMediaSourceCustomization::HandleMediaSourceEntryBoxShouldFilterAsset)
 			];
 	}
 
@@ -151,14 +151,7 @@ void FPlatformMediaSourceCustomization::SetPlatformMediaSourcesValue(FString Pla
 /* FPlatformMediaSourceCustomization callbacks
  *****************************************************************************/
 
-bool FPlatformMediaSourceCustomization::HandleShouldFilterAsset(const FAssetData& AssetData)
-{
-	// Don't allow nesting platform media sources.
-	UClass* AssetClass = FindObject<UClass>(ANY_PACKAGE, *AssetData.AssetClass.ToString());
-	return AssetClass->IsChildOf(UPlatformMediaSource::StaticClass());
-}
-
-void FPlatformMediaSourceCustomization::HandleMediaSourcePropertyEntryBoxChanged(const FAssetData& AssetData, FString PlatformName)
+void FPlatformMediaSourceCustomization::HandleMediaSourceEntryBoxChanged(const FAssetData& AssetData, FString PlatformName)
 {
 	TArray<UObject*> OuterObjects;
 	{
@@ -178,7 +171,7 @@ void FPlatformMediaSourceCustomization::HandleMediaSourcePropertyEntryBoxChanged
 }
 
 
-FString FPlatformMediaSourceCustomization::HandleMediaSourcePropertyEntryObjectPath(FString PlatformName) const
+FString FPlatformMediaSourceCustomization::HandleMediaSourceEntryBoxObjectPath(FString PlatformName) const
 {
 	TArray<UObject*> OuterObjects;
 	{
@@ -209,4 +202,12 @@ FString FPlatformMediaSourceCustomization::HandleMediaSourcePropertyEntryObjectP
 }
 
 
+bool FPlatformMediaSourceCustomization::HandleMediaSourceEntryBoxShouldFilterAsset(const FAssetData& AssetData)
+{
+	// Don't allow nesting platform media sources.
+	UClass* AssetClass = FindObject<UClass>(ANY_PACKAGE, *AssetData.AssetClass.ToString());
+	return AssetClass->IsChildOf(UPlatformMediaSource::StaticClass());
+}
+
+
 #undef LOCTEXT_NAMESPACE
diff --git a/Engine/Plugins/Media/MediaPlayerEditor/Source/MediaPlayerEditor/Private/Customizations/PlatformMediaSourceCustomization.h b/Engine/Plugins/Media/MediaPlayerEditor/Source/MediaPlayerEditor/Private/Customizations/PlatformMediaSourceCustomization.h
--- a/Engine/Plugins/Media/MediaPlayerEditor/Source/MediaPlayerEditor/Private/Customizations/PlatformMediaSourceCustomization.h
+++ b/Engine/Plugins/Media/MediaPlayerEditor/Source/MediaPlayerEditor/Private/Customizations/PlatformMediaSourceCustomization.h
@@ -53,14 +53,14 @@ protected:
 
 private:
 
-	/** Callback for filtering media source assets. */
-	bool HandleShouldFilterAsset(const FAssetData& AssetData);
-
 	/** Callback for when a per-platform media source property changed. */
-	void HandleMediaSourcePropertyEntryBoxChanged(const FAssetData& AssetData, FString PlatformName);
+	void HandleMediaSourceEntryBoxChanged(const FAssetData& AssetData, FString PlatformName);
 
 	/** Callback for getting the currently selected object in a per-platform media source property. */
-	FString HandleMediaSourcePropertyEntryObjectPath(FString PlatformName) const;
+	FString HandleMediaSourceEntryBoxObjectPath(FString PlatformName) const;
+
+	/** Callback for filtering media source assets. */
+	bool HandleMediaSourceEntryBoxShouldFilterAsset(const FAssetData& AssetData);
 
 private:
 
diff --git a/Engine/Plugins/Media/MediaPlayerEditor/Source/MediaPlayerEditor/Private/Factories/MediaPlayerFactoryNew.cpp b/Engine/Plugins/Media/MediaPlayerEditor/Source/MediaPlayerEditor/Private/Factories/MediaPlayerFactoryNew.cpp
--- a/Engine/Plugins/Media/MediaPlayerEditor/Source/MediaPlayerEditor/Private/Factories/MediaPlayerFactoryNew.cpp
+++ b/Engine/Plugins/Media/MediaPlayerEditor/Source/MediaPlayerEditor/Private/Factories/MediaPlayerFactoryNew.cpp
@@ -1,28 +1,28 @@
 // Copyright 1998-2017 Epic Games, Inc. All Rights Reserved.
 
-#include "Factories/MediaPlayerFactoryNew.h"
-#include "Modules/ModuleManager.h"
-#include "Misc/PackageName.h"
-#include "IAssetTools.h"
+#include "MediaPlayerFactoryNew.h"
+
 #include "AssetToolsModule.h"
-#include "Layout/Visibility.h"
+#include "Editor.h"
+#include "EditorStyleSet.h"
+#include "IAssetTools.h"
 #include "Input/Reply.h"
-#include "Widgets/DeclarativeSyntaxSupport.h"
-#include "Widgets/SCompoundWidget.h"
+#include "Misc/PackageName.h"
+#include "Modules/ModuleManager.h"
+#include "Layout/Visibility.h"
+#include "MediaPlayer.h"
+#include "MediaTexture.h"
 #include "Styling/SlateTypes.h"
+#include "Widgets/DeclarativeSyntaxSupport.h"
 #include "Widgets/SBoxPanel.h"
+#include "Widgets/SCompoundWidget.h"
 #include "Widgets/SWindow.h"
-#include "Widgets/Layout/SBorder.h"
-#include "Widgets/Text/STextBlock.h"
-#include "Widgets/Layout/SUniformGridPanel.h"
 #include "Widgets/Input/SButton.h"
 #include "Widgets/Input/SCheckBox.h"
-#include "EditorStyleSet.h"
-#include "Editor.h"
-#include "MediaPlayer.h"
-#include "MediaSoundWave.h"
-#include "MediaTexture.h"
-#include "Factories/MediaSoundWaveFactoryNew.h"
+#include "Widgets/Layout/SBorder.h"
+#include "Widgets/Layout/SUniformGridPanel.h"
+#include "Widgets/Text/STextBlock.h"
+
 #include "Factories/MediaTextureFactoryNew.h"
 
 
@@ -71,21 +71,6 @@ public:
 												.Text(LOCTEXT("CreateAdditionalAssetsLabel", "Additional assets to create and link to the Media Player:"))
 										]
 
-									+ SVerticalBox::Slot()
-										.Padding(0.0f, 8.0f, 0.0f, 0.0f)
-										[
-											SNew(SCheckBox)
-												.IsChecked(Options->CreateSoundWave ? ECheckBoxState::Checked : ECheckBoxState::Unchecked)
-												.OnCheckStateChanged_Lambda([this](ECheckBoxState CheckBoxState) {
-													Options->CreateSoundWave = (CheckBoxState == ECheckBoxState::Checked);
-												})
-												.Content()
-												[
-													SNew(STextBlock)
-														.Text(LOCTEXT("CreateSoundWaveLabel", "Audio output SoundWave asset"))
-												]
-										]
-
 									+ SVerticalBox::Slot()
 										.Padding(0.0f, 6.0f, 0.0f, 0.0f)
 										[
@@ -172,7 +157,6 @@ UMediaPlayerFactoryNew::UMediaPlayerFactoryNew( const FObjectInitializer& Object
 
 bool UMediaPlayerFactoryNew::ConfigureProperties()
 {
-	Options.CreateSoundWave = false;
 	Options.CreateVideoTexture = false;
 	Options.OkClicked = false;
 
@@ -193,7 +177,7 @@ UObject* UMediaPlayerFactoryNew::FactoryCreateNew(UClass* InClass, UObject* InPa
 {
 	auto NewMediaPlayer = NewObject<UMediaPlayer>(InParent, InClass, InName, Flags);
 
-	if ((NewMediaPlayer != nullptr) && (Options.CreateSoundWave || Options.CreateVideoTexture))
+	if ((NewMediaPlayer != nullptr) && Options.CreateVideoTexture)
 	{
 		IAssetTools& AssetTools = FModuleManager::Get().LoadModuleChecked<FAssetToolsModule>("AssetTools").Get();
 		const FString ParentName = InParent->GetOutermost()->GetName();
@@ -201,30 +185,14 @@ UObject* UMediaPlayerFactoryNew::FactoryCreateNew(UClass* InClass, UObject* InPa
 		FString OutAssetName;
 		FString OutPackageName;
 
-		if (Options.CreateSoundWave)
-		{
-			AssetTools.CreateUniqueAssetName(ParentName, TEXT("_Sound"), OutPackageName, OutAssetName);
-			const FString PackagePath = FPackageName::GetLongPackagePath(OutPackageName);
-			auto Factory = NewObject<UMediaSoundWaveFactoryNew>();
-			auto SoundWave = Cast<UMediaSoundWave>(AssetTools.CreateAsset(OutAssetName, PackagePath, UMediaSoundWave::StaticClass(), Factory));
-
-			if (SoundWave != nullptr)
-			{
-				NewMediaPlayer->SetSoundWave(SoundWave);
-			}
-		}
+		AssetTools.CreateUniqueAssetName(ParentName, TEXT("_Video"), OutPackageName, OutAssetName);
+		const FString PackagePath = FPackageName::GetLongPackagePath(OutPackageName);
+		auto Factory = NewObject<UMediaTextureFactoryNew>();
+		auto VideoTexture = Cast<UMediaTexture>(AssetTools.CreateAsset(OutAssetName, PackagePath, UMediaTexture::StaticClass(), Factory));
 
-		if (Options.CreateVideoTexture)
+		if (VideoTexture != nullptr)
 		{
-			AssetTools.CreateUniqueAssetName(ParentName, TEXT("_Video"), OutPackageName, OutAssetName);
-			const FString PackagePath = FPackageName::GetLongPackagePath(OutPackageName);
-			auto Factory = NewObject<UMediaTextureFactoryNew>();
-			auto VideoTexture = Cast<UMediaTexture>(AssetTools.CreateAsset(OutAssetName, PackagePath, UMediaTexture::StaticClass(), Factory));
-
-			if (VideoTexture != nullptr)
-			{
-				NewMediaPlayer->SetVideoTexture(VideoTexture);
-			}
+			VideoTexture->MediaPlayer = NewMediaPlayer;
 		}
 	}
 
diff --git a/Engine/Plugins/Media/MediaPlayerEditor/Source/MediaPlayerEditor/Private/Factories/MediaPlayerFactoryNew.h b/Engine/Plugins/Media/MediaPlayerEditor/Source/MediaPlayerEditor/Private/Factories/MediaPlayerFactoryNew.h
--- a/Engine/Plugins/Media/MediaPlayerEditor/Source/MediaPlayerEditor/Private/Factories/MediaPlayerFactoryNew.h
+++ b/Engine/Plugins/Media/MediaPlayerEditor/Source/MediaPlayerEditor/Private/Factories/MediaPlayerFactoryNew.h
@@ -12,7 +12,6 @@
  */
 struct FMediaPlayerFactoryNewOptions
 {
-	bool CreateSoundWave;
 	bool CreateVideoTexture;
 	bool OkClicked;
 };
diff --git a/Engine/Plugins/Media/MediaPlayerEditor/Source/MediaPlayerEditor/Private/MediaPlayerEditorModule.cpp b/Engine/Plugins/Media/MediaPlayerEditor/Source/MediaPlayerEditor/Private/MediaPlayerEditorModule.cpp
--- a/Engine/Plugins/Media/MediaPlayerEditor/Source/MediaPlayerEditor/Private/MediaPlayerEditorModule.cpp
+++ b/Engine/Plugins/Media/MediaPlayerEditor/Source/MediaPlayerEditor/Private/MediaPlayerEditorModule.cpp
@@ -16,7 +16,6 @@
 #include "BaseMediaSource.h"
 #include "FileMediaSource.h"
 #include "MediaPlayer.h"
-#include "MediaSoundWave.h"
 #include "MediaTexture.h"
 #include "PlatformMediaSource.h"
 
@@ -25,11 +24,9 @@
 #include "AssetTools/FileMediaSourceActions.h"
 #include "AssetTools/MediaPlayerActions.h"
 #include "AssetTools/MediaPlaylistActions.h"
-#include "AssetTools/MediaSoundWaveActions.h"
 
 #include "Customizations/BaseMediaSourceCustomization.h"
 #include "Customizations/FileMediaSourceCustomization.h"
-#include "Customizations/MediaSoundWaveCustomization.h"
 #include "Customizations/MediaTextureCustomization.h"
 #include "Customizations/PlatformMediaSourceCustomization.h"
 
@@ -105,7 +102,6 @@ protected:
 		RegisterAssetTypeAction(AssetTools, MakeShareable(new FFileMediaSourceActions(Style.ToSharedRef())));
 		RegisterAssetTypeAction(AssetTools, MakeShareable(new FMediaPlayerActions(Style.ToSharedRef())));
 		RegisterAssetTypeAction(AssetTools, MakeShareable(new FMediaPlaylistActions(Style.ToSharedRef())));
-		RegisterAssetTypeAction(AssetTools, MakeShareable(new FMediaSoundWaveActions));
 		RegisterAssetTypeAction(AssetTools, MakeShareable(new FMediaSourceActions));
 		RegisterAssetTypeAction(AssetTools, MakeShareable(new FMediaTextureActions));
 	}
@@ -145,14 +141,12 @@ protected:
 	{
 		BaseMediaSourceName = UBaseMediaSource::StaticClass()->GetFName();
 		FileMediaSourceName = UFileMediaSource::StaticClass()->GetFName();
-		MediaSoundWaveName = UMediaSoundWave::StaticClass()->GetFName();
 		MediaTextureName = UMediaTexture::StaticClass()->GetFName();
 		PlatformMediaSourceName = UPlatformMediaSource::StaticClass()->GetFName();
 
 		FPropertyEditorModule& PropertyModule = FModuleManager::LoadModuleChecked<FPropertyEditorModule>("PropertyEditor");
 		{
 			PropertyModule.RegisterCustomClassLayout(FileMediaSourceName, FOnGetDetailCustomizationInstance::CreateStatic(&FFileMediaSourceCustomization::MakeInstance));
-			PropertyModule.RegisterCustomClassLayout(MediaSoundWaveName, FOnGetDetailCustomizationInstance::CreateStatic(&FMediaSoundWaveCustomization::MakeInstance));
 			PropertyModule.RegisterCustomClassLayout(MediaTextureName, FOnGetDetailCustomizationInstance::CreateStatic(&FMediaTextureCustomization::MakeInstance));
 			PropertyModule.RegisterCustomClassLayout(PlatformMediaSourceName, FOnGetDetailCustomizationInstance::CreateStatic(&FPlatformMediaSourceCustomization::MakeInstance));
 			PropertyModule.RegisterCustomClassLayout(BaseMediaSourceName, FOnGetDetailCustomizationInstance::CreateStatic(&FBaseMediaSourceCustomization::MakeInstance));
@@ -166,7 +160,6 @@ protected:
 		{
 			PropertyModule.UnregisterCustomClassLayout(BaseMediaSourceName);
 			PropertyModule.UnregisterCustomClassLayout(FileMediaSourceName);
-			PropertyModule.UnregisterCustomClassLayout(MediaSoundWaveName);
 			PropertyModule.UnregisterCustomClassLayout(MediaTextureName);
 			PropertyModule.UnregisterCustomClassLayout(PlatformMediaSourceName);
 		}
@@ -276,7 +269,6 @@ private:
 	/** Class names. */
 	FName BaseMediaSourceName;
 	FName FileMediaSourceName;
-	FName MediaSoundWaveName;
 	FName MediaTextureName;
 	FName PlatformMediaSourceName;
 };
diff --git a/Engine/Plugins/Media/MediaPlayerEditor/Source/MediaPlayerEditor/Private/Shared/MediaPlayerEditorStyle.cpp b/Engine/Plugins/Media/MediaPlayerEditor/Source/MediaPlayerEditor/Private/Shared/MediaPlayerEditorStyle.cpp
new file mode 100644
--- /dev/null
+++ b/Engine/Plugins/Media/MediaPlayerEditor/Source/MediaPlayerEditor/Private/Shared/MediaPlayerEditorStyle.cpp
@@ -0,0 +1,96 @@
+// Copyright 1998-2017 Epic Games, Inc. All Rights Reserved.
+
+#include "MediaPlayerEditorStyle.h"
+
+#include "Brushes/SlateBorderBrush.h"
+#include "Brushes/SlateBoxBrush.h"
+#include "Brushes/SlateImageBrush.h"
+#include "Fonts/SlateFontInfo.h"
+#include "Math/Vector2D.h"
+#include "Misc/Paths.h"
+#include "Styling/SlateStyleRegistry.h"
+#include "Styling/SlateTypes.h"
+
+
+#define IMAGE_BRUSH(RelativePath, ...) FSlateImageBrush(RootToContentDir(RelativePath, TEXT(".png")), __VA_ARGS__)
+#define BOX_BRUSH(RelativePath, ...) FSlateBoxBrush(RootToContentDir(RelativePath, TEXT(".png")), __VA_ARGS__)
+#define BORDER_BRUSH(RelativePath, ...) FSlateBorderBrush(RootToContentDir(RelativePath, TEXT(".png")), __VA_ARGS__)
+#define TTF_FONT(RelativePath, ...) FSlateFontInfo(RootToContentDir(RelativePath, TEXT(".ttf")), __VA_ARGS__)
+#define OTF_FONT(RelativePath, ...) FSlateFontInfo(RootToContentDir(RelativePath, TEXT(".otf")), __VA_ARGS__)
+
+
+/* FMediaPlayerEditorStyle structors
+ *****************************************************************************/
+
+FMediaPlayerEditorStyle::FMediaPlayerEditorStyle()
+	: FSlateStyleSet("MediaPlayerEditorStyle")
+{
+	const FVector2D Icon8x8(8.0f, 8.0f);
+	const FVector2D Icon12x12(12.0f, 12.0f);
+	const FVector2D Icon16x16(16.0f, 16.0f);
+	const FVector2D Icon20x20(20.0f, 20.0f);
+	const FVector2D Icon40x40(40.0f, 40.0f);
+
+	SetContentRoot(FPaths::EnginePluginsDir() / TEXT("Media/MediaPlayerEditor/Content"));
+
+	// buttons
+	Set("MediaPlayerEditor.SourceButton", new IMAGE_BRUSH("btn_source_12x", Icon12x12));
+	Set("MediaPlayerEditor.GoButton", new IMAGE_BRUSH("btn_go_12x", Icon12x12));
+	Set("MediaPlayerEditor.ReloadButton", new IMAGE_BRUSH("btn_reload_12x", Icon12x12));
+	Set("MediaPlayerEditor.SettingsButton", new IMAGE_BRUSH("btn_settings_16x", Icon12x12));
+
+	// misc
+	Set("MediaPlayerEditor.DragDropBorder", new BOX_BRUSH("border_dragdrop", 0.5f));
+	Set("MediaPlayerEditor.FileMediaSourcePrecached", new IMAGE_BRUSH("filemediasource_precached", Icon12x12));
+	Set("MediaPlayerEditor.MediaSourceOpened", new IMAGE_BRUSH("mediasource_opened", Icon8x8));
+
+	// tabs
+	Set("MediaPlayerEditor.Tabs.Info", new IMAGE_BRUSH("tab_info_16x", Icon16x16));
+	Set("MediaPlayerEditor.Tabs.Media", new IMAGE_BRUSH("tab_media_16x", Icon16x16));
+	Set("MediaPlayerEditor.Tabs.Player", new IMAGE_BRUSH("tab_player_16x", Icon16x16));
+	Set("MediaPlayerEditor.Tabs.Playlist", new IMAGE_BRUSH("tab_playlist_16x", Icon16x16));
+	Set("MediaPlayerEditor.Tabs.Stats", new IMAGE_BRUSH("tab_stats_16x", Icon16x16));
+
+	// toolbar icons
+	Set("MediaPlayerEditor.CloseMedia", new IMAGE_BRUSH("icon_eject_40x", Icon40x40));
+	Set("MediaPlayerEditor.ForwardMedia", new IMAGE_BRUSH("icon_forward_40x", Icon40x40));
+	Set("MediaPlayerEditor.ForwardMedia.Small", new IMAGE_BRUSH("icon_forward_40x", Icon20x20));
+	Set("MediaPlayerEditor.NextMedia", new IMAGE_BRUSH("icon_step_40x", Icon40x40));
+	Set("MediaPlayerEditor.NextMedia.Small", new IMAGE_BRUSH("icon_step_40x", Icon20x20));
+	Set("MediaPlayerEditor.PauseMedia", new IMAGE_BRUSH("icon_pause_40x", Icon40x40));
+	Set("MediaPlayerEditor.PauseMedia.Small", new IMAGE_BRUSH("icon_pause_40x", Icon20x20));
+	Set("MediaPlayerEditor.PlayMedia", new IMAGE_BRUSH("icon_play_40x", Icon40x40));
+	Set("MediaPlayerEditor.PlayMedia.Small", new IMAGE_BRUSH("icon_play_40x", Icon20x20));
+	Set("MediaPlayerEditor.PreviousMedia", new IMAGE_BRUSH("icon_step_back_40x", Icon40x40));
+	Set("MediaPlayerEditor.PreviousMedia.Small", new IMAGE_BRUSH("icon_step_back_40x", Icon20x20));
+	Set("MediaPlayerEditor.ReverseMedia", new IMAGE_BRUSH("icon_reverse_40x", Icon40x40));
+	Set("MediaPlayerEditor.ReverseMedia.Small", new IMAGE_BRUSH("icon_reverse_40x", Icon20x20));
+	Set("MediaPlayerEditor.RewindMedia", new IMAGE_BRUSH("icon_rewind_40x", Icon40x40));
+	Set("MediaPlayerEditor.RewindMedia.Small", new IMAGE_BRUSH("icon_rewind_40x", Icon20x20));
+	Set("MediaPlayerEditor.StopMedia", new IMAGE_BRUSH("icon_stop_40x", Icon40x40));
+	Set("MediaPlayerEditor.StopMedia.Small", new IMAGE_BRUSH("icon_stop_40x", Icon20x20));
+
+	// scrubber
+	Set("MediaPlayerEditor.Scrubber", FSliderStyle()
+		.SetNormalBarImage(FSlateColorBrush(FColor::White))
+		.SetDisabledBarImage(FSlateColorBrush(FLinearColor::Gray))
+		.SetNormalThumbImage(IMAGE_BRUSH("scrubber", FVector2D(2.0f, 10.0f)))
+		.SetDisabledThumbImage(IMAGE_BRUSH("scrubber", FVector2D(2.0f, 10.0f)))
+		.SetBarThickness(2.0f)
+	);
+
+	FSlateStyleRegistry::RegisterSlateStyle(*this);
+}
+
+
+FMediaPlayerEditorStyle::~FMediaPlayerEditorStyle()
+{
+	FSlateStyleRegistry::UnRegisterSlateStyle(*this);
+}
+
+
+#undef IMAGE_BRUSH
+#undef BOX_BRUSH
+#undef BORDER_BRUSH
+#undef TTF_FONT
+#undef OTF_FONT
diff --git a/Engine/Plugins/Media/MediaPlayerEditor/Source/MediaPlayerEditor/Private/Shared/MediaPlayerEditorStyle.h b/Engine/Plugins/Media/MediaPlayerEditor/Source/MediaPlayerEditor/Private/Shared/MediaPlayerEditorStyle.h
--- a/Engine/Plugins/Media/MediaPlayerEditor/Source/MediaPlayerEditor/Private/Shared/MediaPlayerEditorStyle.h
+++ b/Engine/Plugins/Media/MediaPlayerEditor/Source/MediaPlayerEditor/Private/Shared/MediaPlayerEditorStyle.h
@@ -2,21 +2,8 @@
 
 #pragma once
 
-#include "CoreMinimal.h"
-#include "Fonts/SlateFontInfo.h"
-#include "Misc/Paths.h"
-#include "Brushes/SlateBorderBrush.h"
-#include "Brushes/SlateBoxBrush.h"
-#include "Brushes/SlateImageBrush.h"
-#include "Styling/SlateStyleRegistry.h"
 #include "Styling/SlateStyle.h"
 
-#define IMAGE_BRUSH(RelativePath, ...) FSlateImageBrush(RootToContentDir(RelativePath, TEXT(".png")), __VA_ARGS__)
-#define BOX_BRUSH(RelativePath, ...) FSlateBoxBrush(RootToContentDir(RelativePath, TEXT(".png")), __VA_ARGS__)
-#define BORDER_BRUSH(RelativePath, ...) FSlateBorderBrush(RootToContentDir(RelativePath, TEXT(".png")), __VA_ARGS__)
-#define TTF_FONT(RelativePath, ...) FSlateFontInfo(RootToContentDir(RelativePath, TEXT(".ttf")), __VA_ARGS__)
-#define OTF_FONT(RelativePath, ...) FSlateFontInfo(RootToContentDir(RelativePath, TEXT(".otf")), __VA_ARGS__)
-
 
 /**
  * Implements the visual style of the media player editor UI.
@@ -27,66 +14,8 @@ class FMediaPlayerEditorStyle
 public:
 
 	/** Default constructor. */
-	FMediaPlayerEditorStyle()
-		: FSlateStyleSet("MediaPlayerEditorStyle")
-	{
-		const FVector2D Icon8x8(8.0f, 8.0f);
-		const FVector2D Icon12x12(12.0f, 12.0f);
-		const FVector2D Icon16x16(16.0f, 16.0f);
-		const FVector2D Icon20x20(20.0f, 20.0f);
-		const FVector2D Icon40x40(40.0f, 40.0f);
-
-		SetContentRoot(FPaths::EnginePluginsDir() / TEXT("Media/MediaPlayerEditor/Content"));
-
-		// buttons
-		Set("MediaPlayerEditor.GoButton", new IMAGE_BRUSH("btn_go_12x", Icon12x12));
-		Set("MediaPlayerEditor.ReloadButton", new IMAGE_BRUSH("btn_reload_12x", Icon12x12));
-		Set("MediaPlayerEditor.SettingsButton", new IMAGE_BRUSH("btn_settings_16x", Icon12x12));
-
-		// misc
-		Set("MediaPlayerEditor.DragDropBorder", new BOX_BRUSH("border_dragdrop", 0.5f));
-		Set("MediaPlayerEditor.FileMediaSourcePrecached", new IMAGE_BRUSH("filemediasource_precached", Icon12x12));
-		Set("MediaPlayerEditor.MediaSourceOpened", new IMAGE_BRUSH("mediasource_opened", Icon8x8));
-
-		// tabs
-		Set("MediaPlayerEditor.Tabs.Info", new IMAGE_BRUSH("tab_info_16x", Icon16x16));
-		Set("MediaPlayerEditor.Tabs.Media", new IMAGE_BRUSH("tab_media_16x", Icon16x16));
-		Set("MediaPlayerEditor.Tabs.Player", new IMAGE_BRUSH("tab_player_16x", Icon16x16));
-		Set("MediaPlayerEditor.Tabs.Playlist", new IMAGE_BRUSH("tab_playlist_16x", Icon16x16));
-		Set("MediaPlayerEditor.Tabs.Stats", new IMAGE_BRUSH("tab_stats_16x", Icon16x16));
-
-		// toolbar icons
-		Set("MediaPlayerEditor.CloseMedia", new IMAGE_BRUSH("icon_eject_40x", Icon40x40));
-		Set("MediaPlayerEditor.ForwardMedia", new IMAGE_BRUSH("icon_forward_40x", Icon40x40));
-		Set("MediaPlayerEditor.ForwardMedia.Small", new IMAGE_BRUSH("icon_forward_40x", Icon20x20));
-		Set("MediaPlayerEditor.NextMedia", new IMAGE_BRUSH("icon_step_40x", Icon40x40));
-		Set("MediaPlayerEditor.NextMedia.Small", new IMAGE_BRUSH("icon_step_40x", Icon20x20));
-		Set("MediaPlayerEditor.PauseMedia", new IMAGE_BRUSH("icon_pause_40x", Icon40x40));
-		Set("MediaPlayerEditor.PauseMedia.Small", new IMAGE_BRUSH("icon_pause_40x", Icon20x20));
-		Set("MediaPlayerEditor.PlayMedia", new IMAGE_BRUSH("icon_play_40x", Icon40x40));
-		Set("MediaPlayerEditor.PlayMedia.Small", new IMAGE_BRUSH("icon_play_40x", Icon20x20));
-		Set("MediaPlayerEditor.PreviousMedia", new IMAGE_BRUSH("icon_step_back_40x", Icon40x40));
-		Set("MediaPlayerEditor.PreviousMedia.Small", new IMAGE_BRUSH("icon_step_back_40x", Icon20x20));
-		Set("MediaPlayerEditor.ReverseMedia", new IMAGE_BRUSH("icon_reverse_40x", Icon40x40));
-		Set("MediaPlayerEditor.ReverseMedia.Small", new IMAGE_BRUSH("icon_reverse_40x", Icon20x20));
-		Set("MediaPlayerEditor.RewindMedia", new IMAGE_BRUSH("icon_rewind_40x", Icon40x40));
-		Set("MediaPlayerEditor.RewindMedia.Small", new IMAGE_BRUSH("icon_rewind_40x", Icon20x20));
-		Set("MediaPlayerEditor.StopMedia", new IMAGE_BRUSH("icon_stop_40x", Icon40x40));
-		Set("MediaPlayerEditor.StopMedia.Small", new IMAGE_BRUSH("icon_stop_40x", Icon20x20));
-		
-		FSlateStyleRegistry::RegisterSlateStyle(*this);
-	 }
+	FMediaPlayerEditorStyle();
 
 	 /** Destructor. */
-	 ~FMediaPlayerEditorStyle()
-	 {
-		FSlateStyleRegistry::UnRegisterSlateStyle(*this);
-	 }
+	~FMediaPlayerEditorStyle();
 };
-
-
-#undef IMAGE_BRUSH
-#undef BOX_BRUSH
-#undef BORDER_BRUSH
-#undef TTF_FONT
-#undef OTF_FONT
diff --git a/Engine/Plugins/Media/MediaPlayerEditor/Source/MediaPlayerEditor/Private/Toolkits/MediaPlayerEditorToolkit.cpp b/Engine/Plugins/Media/MediaPlayerEditor/Source/MediaPlayerEditor/Private/Toolkits/MediaPlayerEditorToolkit.cpp
--- a/Engine/Plugins/Media/MediaPlayerEditor/Source/MediaPlayerEditor/Private/Toolkits/MediaPlayerEditorToolkit.cpp
+++ b/Engine/Plugins/Media/MediaPlayerEditor/Source/MediaPlayerEditor/Private/Toolkits/MediaPlayerEditorToolkit.cpp
@@ -294,7 +294,7 @@ void FMediaPlayerEditorToolkit::BindCommands()
 	ToolkitCommands->MapAction(
 		Commands.NextMedia,
 		FExecuteAction::CreateLambda([this]{ MediaPlayer->Next(); }),
-		FCanExecuteAction::CreateLambda([this]{ return (MediaPlayer->GetPlaylist() != nullptr) && (MediaPlayer->GetPlaylist()->Num() > 1); })
+		FCanExecuteAction::CreateLambda([this]{ return (MediaPlayer->GetPlaylistRef().Num() > 1); })
 	);
 
 	ToolkitCommands->MapAction(
@@ -312,7 +312,7 @@ void FMediaPlayerEditorToolkit::BindCommands()
 	ToolkitCommands->MapAction(
 		Commands.PreviousMedia,
 		FExecuteAction::CreateLambda([this]{ MediaPlayer->Previous(); }),
-		FCanExecuteAction::CreateLambda([this]{ return (MediaPlayer->GetPlaylist() != nullptr) && (MediaPlayer->GetPlaylist()->Num() > 1); })
+		FCanExecuteAction::CreateLambda([this]{ return (MediaPlayer->GetPlaylistRef().Num() > 1); })
 	);
 
 	ToolkitCommands->MapAction(
diff --git a/Engine/Plugins/Media/MediaPlayerEditor/Source/MediaPlayerEditor/Private/Widgets/SMediaPlayerEditorInfo.cpp b/Engine/Plugins/Media/MediaPlayerEditor/Source/MediaPlayerEditor/Private/Widgets/SMediaPlayerEditorInfo.cpp
--- a/Engine/Plugins/Media/MediaPlayerEditor/Source/MediaPlayerEditor/Private/Widgets/SMediaPlayerEditorInfo.cpp
+++ b/Engine/Plugins/Media/MediaPlayerEditor/Source/MediaPlayerEditor/Private/Widgets/SMediaPlayerEditorInfo.cpp
@@ -1,21 +1,66 @@
 // Copyright 1998-2017 Epic Games, Inc. All Rights Reserved.
 
-#include "Widgets/SMediaPlayerEditorInfo.h"
-#include "Widgets/Text/STextBlock.h"
-#include "Widgets/SBoxPanel.h"
-#include "Widgets/Layout/SBorder.h"
-#include "Widgets/Input/SButton.h"
-#include "Widgets/Layout/SScrollBox.h"
+#include "SMediaPlayerEditorInfo.h"
+
 #include "EditorStyleSet.h"
+#include "HAL/PlatformApplicationMisc.h"
 #include "IMediaControls.h"
-#include "IMediaPlayer.h"
 #include "MediaPlayer.h"
-#include "HAL/PlatformApplicationMisc.h"
+#include "MediaPlayerFacade.h"
+#include "Widgets/SBoxPanel.h"
+#include "Widgets/Input/SButton.h"
+#include "Widgets/Layout/SBorder.h"
+#include "Widgets/Layout/SScrollBox.h"
+#include "Widgets/Text/STextBlock.h"
 
 
 #define LOCTEXT_NAMESPACE "SMediaPlayerEditorInfo"
 
 
+/* Local helpers
+ *****************************************************************************/
+
+/**
+ * Convert a set of play rates to human readable text.
+ *
+ * @param Rates The ratest to convert.
+ * @return The corresponding text representation.
+ */
+FText RatesToText(TRangeSet<float> Rates)
+{
+	static const FNumberFormattingOptions Options = FNumberFormattingOptions()
+		.SetMinimumFractionalDigits(1)
+		.SetMaximumFractionalDigits(1)
+		.SetUseGrouping(false);
+
+	TArray<TRange<float>> OutRanges;
+	Rates.GetRanges(OutRanges);
+
+	FString String;
+
+	for (int32 RangeIndex = 0; RangeIndex < OutRanges.Num(); ++RangeIndex)
+	{
+		if (RangeIndex > 0)
+		{
+			String += TEXT(", ");
+		}
+
+		const TRange<float>& Range = OutRanges[RangeIndex];
+
+		if (Range.IsDegenerate())
+		{
+			String += FText::Format(LOCTEXT("DegenerateRateFormat", "{0}"), FText::AsNumber(Range.GetLowerBoundValue(), &Options)).ToString();
+		}
+		else
+		{
+			String += FText::Format(LOCTEXT("NormalRatesFormat", "{0} to {1}"), FText::AsNumber(Range.GetLowerBoundValue(), &Options), FText::AsNumber(Range.GetUpperBoundValue(), &Options)).ToString();
+		}
+	}
+
+	return FText::FromString(String);
+}
+
+
 /* SMediaPlayerEditorInfo interface
  *****************************************************************************/
 
@@ -71,58 +116,38 @@ void SMediaPlayerEditorInfo::HandleMediaPlayerMediaEvent(EMediaEvent Event)
 
 	if ((Event == EMediaEvent::MediaOpened) || (Event == EMediaEvent::MediaOpenFailed) || (Event == EMediaEvent::TracksChanged))
 	{
-		FMediaPlayerBase& Player = MediaPlayer->GetBasePlayer();
+		TSharedRef<FMediaPlayerFacade, ESPMode::ThreadSafe> PlayerFacade = MediaPlayer->GetPlayerFacade();
 
-		if (Player.GetUrl().IsEmpty())
-		{
-			InfoTextBlock->SetText(NoMediaText);
-		}
-		else
+		if (!PlayerFacade->GetUrl().IsEmpty())
 		{
-			TRange<float> ThinnedForwardRates = Player.GetForwardRates(false);
-			TRange<float> UnthinnedForwardRates = Player.GetForwardRates(true);
-			TRange<float> ThinnedReverseRates = Player.GetReverseRates(false);
-			TRange<float> UnthinnedReverseRates = Player.GetReverseRates(true);
+			const FText ThinnedRatesText = RatesToText(PlayerFacade->GetSupportedRates(false));
+			const FText UnthinnedRatesText = RatesToText(PlayerFacade->GetSupportedRates(true));
 
 			FFormatNamedArguments Arguments;
 			{
-				Arguments.Add(TEXT("PlayerName"), FText::FromName(Player.GetPlayerName()));
-				Arguments.Add(TEXT("SupportsScrubbing"), Player.SupportsScrubbing() ? GYes : GNo);
-				Arguments.Add(TEXT("SupportsSeeking"), Player.SupportsSeeking() ? GYes : GNo);
-				Arguments.Add(TEXT("PlayerInfo"), FText::FromString(Player.GetInfo()));
-
-				Arguments.Add(TEXT("ForwardRateThinned"), ThinnedForwardRates.IsEmpty()
-					? LOCTEXT("RateNotSupported", "Not supported")
-					: FText::Format(LOCTEXT("RatesFormat", "{0} to {1}"), FText::AsNumber(ThinnedForwardRates.GetLowerBoundValue()), FText::AsNumber(ThinnedForwardRates.GetUpperBoundValue()))
-				);
+				Arguments.Add(TEXT("PlayerName"), FText::FromName(PlayerFacade->GetPlayerName()));
+				Arguments.Add(TEXT("SupportsScrubbing"), PlayerFacade->CanScrub() ? GYes : GNo);
+				Arguments.Add(TEXT("SupportsSeeking"), PlayerFacade->CanSeek() ? GYes : GNo);
+				Arguments.Add(TEXT("PlayerInfo"), FText::FromString(PlayerFacade->GetInfo()));
 
-				Arguments.Add(TEXT("ForwardRateUnthinned"), UnthinnedForwardRates.IsEmpty()
+				Arguments.Add(TEXT("ThinnedRates"), ThinnedRatesText.IsEmpty()
 					? LOCTEXT("RateNotSupported", "Not supported")
-					: FText::Format(LOCTEXT("RatesFormat", "{0} to {1}"), FText::AsNumber(UnthinnedForwardRates.GetLowerBoundValue()), FText::AsNumber(UnthinnedForwardRates.GetUpperBoundValue()))
+					: ThinnedRatesText
 				);
 
-				Arguments.Add(TEXT("ReverseRateThinned"), ThinnedReverseRates.IsEmpty()
+				Arguments.Add(TEXT("UnthinnedRates"), UnthinnedRatesText.IsEmpty()
 					? LOCTEXT("RateNotSupported", "Not supported")
-					: FText::Format(LOCTEXT("RatesFormat", "{0} to {1}"), FText::AsNumber(ThinnedReverseRates.GetLowerBoundValue()), FText::AsNumber(ThinnedReverseRates.GetUpperBoundValue()))
+					: UnthinnedRatesText
 				);
-
-				Arguments.Add(TEXT("ReverseRateUnthinned"), UnthinnedReverseRates.IsEmpty()
-					? LOCTEXT("RateNotSupported", "Not supported")
-					: FText::Format(LOCTEXT("RatesFormat", "{0} to {1}"), FText::AsNumber(UnthinnedReverseRates.GetLowerBoundValue()), FText::AsNumber(UnthinnedReverseRates.GetUpperBoundValue()))
-				);				
 			}
 
 			InfoTextBlock->SetText(
 				FText::Format(LOCTEXT("InfoFormat",
 					"Player: {PlayerName}\n"
 					"\n"
-					"Forward Rates\n"
-					"    Thinned: {ForwardRateThinned}\n"
-					"    Unthinned: {ForwardRateUnthinned}\n"
-					"\n"
-					"Reverse Rates\n"
-					"    Thinned: {ReverseRateThinned}\n"
-					"    Unthinned: {ReverseRateUnthinned}\n"
+					"Play Rates\n"
+					"    Thinned: {ThinnedRates}\n"
+					"    Unthinned: {UnthinnedRates}\n"
 					"\n"
 					"Capabilities\n"
 					"    Scrubbing: {SupportsScrubbing}\n"
@@ -133,6 +158,10 @@ void SMediaPlayerEditorInfo::HandleMediaPlayerMediaEvent(EMediaEvent Event)
 				)
 			);
 		}
+		else
+		{
+			InfoTextBlock->SetText(NoMediaText);
+		}
 	}
 	else if (Event == EMediaEvent::MediaClosed)
 	{
diff --git a/Engine/Plugins/Media/MediaPlayerEditor/Source/MediaPlayerEditor/Private/Widgets/SMediaPlayerEditorMedia.cpp b/Engine/Plugins/Media/MediaPlayerEditor/Source/MediaPlayerEditor/Private/Widgets/SMediaPlayerEditorMedia.cpp
--- a/Engine/Plugins/Media/MediaPlayerEditor/Source/MediaPlayerEditor/Private/Widgets/SMediaPlayerEditorMedia.cpp
+++ b/Engine/Plugins/Media/MediaPlayerEditor/Source/MediaPlayerEditor/Private/Widgets/SMediaPlayerEditorMedia.cpp
@@ -2,22 +2,24 @@
 
 #include "Widgets/SMediaPlayerEditorMedia.h"
 #include "Modules/ModuleManager.h"
-#include "Framework/Commands/UIAction.h"
 #include "Textures/SlateIcon.h"
 #include "IContentBrowserSingleton.h"
 #include "ContentBrowserModule.h"
-#include "Widgets/Layout/SBorder.h"
-#include "Framework/MultiBox/MultiBoxBuilder.h"
 #include "EditorStyleSet.h"
 #include "Editor.h"
 #include "Containers/ArrayBuilder.h"
 #include "Toolkits/AssetEditorManager.h"
-#include "MediaSource.h"
 #include "FileMediaSource.h"
 #include "MediaPlayer.h"
 #include "MediaPlaylist.h"
+#include "MediaSource.h"
+
+#include "Framework/Commands/UIAction.h"
+#include "Framework/MultiBox/MultiBoxBuilder.h"
 #include "Framework/Notifications/NotificationManager.h"
+#include "Widgets/Layout/SBorder.h"
 #include "Widgets/Notifications/SNotificationList.h"
+#include "Widgets/Text/STextBlock.h"
 
 
 #define LOCTEXT_NAMESPACE "SMediaPlayerEditorMedia"
@@ -154,8 +156,8 @@ TSharedPtr<SWidget> SMediaPlayerEditorMedia::HandleAssetPickerGetAssetContextMen
 			);
 
 			MenuBuilder.AddMenuEntry(
-				LOCTEXT("LoadMenuAction", "Open"),
-				LOCTEXT("LoadMenuActionTooltip", "Open this media asset in the player"),
+				LOCTEXT("OpenMenuAction", "Open"),
+				LOCTEXT("OpenMenuActionTooltip", "Open this media asset in the player"),
 				FSlateIcon(Style->GetStyleSetName(), "MediaPlayerEditor.NextMedia.Small"),
 				FUIAction(
 					FExecuteAction::CreateLambda([=]() {
@@ -196,7 +198,7 @@ TSharedPtr<SWidget> SMediaPlayerEditorMedia::HandleAssetPickerGetAssetContextMen
 					FSlateIcon(FEditorStyle::GetStyleSetName(), "SystemWideCommands.FindInContentBrowser"),
 					FUIAction(
 						FExecuteAction::CreateLambda([=]() {
-							FPlatformProcess::ExploreFolder(*FileMediaSource->GetFilePath());
+							FPlatformProcess::ExploreFolder(*FileMediaSource->GetFullPath());
 						}),
 						FCanExecuteAction::CreateLambda([=]() -> bool {
 							return FileMediaSource->Validate();
diff --git a/Engine/Plugins/Media/MediaPlayerEditor/Source/MediaPlayerEditor/Private/Widgets/SMediaPlayerEditorMedia.h b/Engine/Plugins/Media/MediaPlayerEditor/Source/MediaPlayerEditor/Private/Widgets/SMediaPlayerEditorMedia.h
--- a/Engine/Plugins/Media/MediaPlayerEditor/Source/MediaPlayerEditor/Private/Widgets/SMediaPlayerEditorMedia.h
+++ b/Engine/Plugins/Media/MediaPlayerEditor/Source/MediaPlayerEditor/Private/Widgets/SMediaPlayerEditorMedia.h
@@ -2,14 +2,16 @@
 
 #pragma once
 
-#include "CoreMinimal.h"
+#include "CoreTypes.h"
 #include "AssetData.h"
 #include "Styling/ISlateStyle.h"
 #include "Widgets/DeclarativeSyntaxSupport.h"
 #include "Widgets/SCompoundWidget.h"
 
+class SWidget;
 class UMediaPlayer;
 
+
 /**
  * Implements the media library of the MediaPlayer asset editor.
  */
diff --git a/Engine/Plugins/Media/MediaPlayerEditor/Source/MediaPlayerEditor/Private/Widgets/SMediaPlayerEditorOutput.cpp b/Engine/Plugins/Media/MediaPlayerEditor/Source/MediaPlayerEditor/Private/Widgets/SMediaPlayerEditorOutput.cpp
--- a/Engine/Plugins/Media/MediaPlayerEditor/Source/MediaPlayerEditor/Private/Widgets/SMediaPlayerEditorOutput.cpp
+++ b/Engine/Plugins/Media/MediaPlayerEditor/Source/MediaPlayerEditor/Private/Widgets/SMediaPlayerEditorOutput.cpp
@@ -1,91 +1,56 @@
 // Copyright 1998-2017 Epic Games, Inc. All Rights Reserved.
 
-#include "Widgets/SMediaPlayerEditorOutput.h"
-#include "Styling/SlateBrush.h"
+#include "SMediaPlayerEditorOutput.h"
+
+#include "AudioDevice.h"
+#include "Editor.h"
+#include "Engine/Engine.h"
+#include "IMediaEventSink.h"
 #include "Materials/Material.h"
 #include "Materials/MaterialExpressionTextureSample.h"
-#include "IMediaOutput.h"
-#include "IMediaPlayer.h"
 #include "MediaPlayer.h"
-#include "MediaSoundWave.h"
+#include "MediaSoundComponent.h"
 #include "MediaTexture.h"
+#include "Styling/SlateBrush.h"
 #include "Widgets/Images/SImage.h"
-#include "Editor.h"
-#include "AudioDevice.h"
 
 
 /* SMediaPlayerEditorOutput structors
  *****************************************************************************/
 
 SMediaPlayerEditorOutput::SMediaPlayerEditorOutput()
-	: CurrentSoundWave(nullptr)
-	, CurrentTexture(nullptr)
-	, DefaultSoundWave(nullptr)
-	, DefaultTexture(nullptr)
-	, Material(nullptr)
+	: Material(nullptr)
 	, MediaPlayer(nullptr)
-{
-	DefaultSoundWave = NewObject<UMediaSoundWave>(GetTransientPackage(), NAME_None, RF_Transient | RF_Public);
-
-	if (DefaultSoundWave != nullptr)
-	{
-		DefaultSoundWave->AddToRoot();
-
-		AudioComponent = FAudioDevice::CreateComponent(DefaultSoundWave);
-	
-		if (AudioComponent != nullptr)
-		{
-			AudioComponent->SetVolumeMultiplier(1.0f);
-			AudioComponent->SetPitchMultiplier(1.0f);
-			AudioComponent->bAllowSpatialization = false;
-			AudioComponent->bIsUISound = true;
-			AudioComponent->bAutoDestroy = false;
-			AudioComponent->AddToRoot();
-		}
-	}
-}
+	, MediaTexture(nullptr)
+	, SoundComponent(nullptr)
+{ }
 
 
 SMediaPlayerEditorOutput::~SMediaPlayerEditorOutput()
 {
-	if (MediaPlayer != nullptr)
+	if (MediaPlayer.IsValid())
 	{
 		MediaPlayer->OnMediaEvent().RemoveAll(this);
-	
-		// remove default sinks from native player
-		FMediaPlayerBase& Player = MediaPlayer->GetBasePlayer();
-
-		if (MediaPlayer->GetSoundWave() == nullptr)
-		{
-			Player.SetAudioSink(nullptr);
-		}
-
-		if (MediaPlayer->GetVideoTexture() == nullptr)
-		{
-			Player.SetVideoSink(nullptr);
-		}
-	}
-
-	// release default sinks
-	if (AudioComponent != nullptr)
-	{
-		AudioComponent->Stop();
-		AudioComponent->RemoveFromRoot();
+		MediaPlayer.Reset();
 	}
 
 	if (Material != nullptr)
 	{
 		Material->RemoveFromRoot();
+		Material = nullptr;
 	}
 
-	if (DefaultSoundWave != nullptr)
+	if (MediaTexture != nullptr)
 	{
-		DefaultSoundWave->RemoveFromRoot();
+		MediaTexture->RemoveFromRoot();
+		MediaTexture = nullptr;
 	}
 
-	if (DefaultTexture != nullptr)
+	if (SoundComponent != nullptr)
 	{
-		DefaultTexture->RemoveFromRoot();
+		SoundComponent->Stop();
+		SoundComponent->RemoveFromRoot();
+		SoundComponent = nullptr;
 	}
 }
 
@@ -97,6 +62,30 @@ void SMediaPlayerEditorOutput::Construct(const FArguments& InArgs, UMediaPlayer&
 {
 	MediaPlayer = &InMediaPlayer;
 
+	// create media sound component
+	if ((GEngine != nullptr) && GEngine->UseSound())
+	{
+		SoundComponent = NewObject<UMediaSoundComponent>(GetTransientPackage(), NAME_None, RF_Transient | RF_Public);
+
+		if (SoundComponent != nullptr)
+		{
+			SoundComponent->MediaPlayer = &InMediaPlayer;
+			SoundComponent->bIsUISound = true;
+			SoundComponent->Initialize();
+			SoundComponent->AddToRoot();
+		}
+	}
+
+	// create media texture
+	MediaTexture = NewObject<UMediaTexture>(GetTransientPackage(), NAME_None, RF_Transient | RF_Public);
+
+	if (MediaTexture != nullptr)
+	{
+		MediaTexture->MediaPlayer = &InMediaPlayer;
+		MediaTexture->UpdateResource();
+		MediaTexture->AddToRoot();
+	}
+
 	// create wrapper material
 	Material = NewObject<UMaterial>(GetTransientPackage(), NAME_None, RF_Transient | RF_Public);
 
@@ -104,7 +93,8 @@ void SMediaPlayerEditorOutput::Construct(const FArguments& InArgs, UMediaPlayer&
 	{
 		TextureSampler = NewObject<UMaterialExpressionTextureSample>(Material);
 		{
-			TextureSampler->SamplerType = SAMPLERTYPE_Color;
+			TextureSampler->Texture = MediaTexture;
+			TextureSampler->AutoSetSampleType();
 		}
 
 		FExpressionOutput& Output = TextureSampler->GetOutputs()[0];
@@ -120,6 +110,7 @@ void SMediaPlayerEditorOutput::Construct(const FArguments& InArgs, UMediaPlayer&
 
 		Material->Expressions.Add(TextureSampler);
 		Material->MaterialDomain = EMaterialDomain::MD_UI;
+		Material->PostEditChange();
 		Material->AddToRoot();
 	}
 
@@ -129,8 +120,6 @@ void SMediaPlayerEditorOutput::Construct(const FArguments& InArgs, UMediaPlayer&
 		MaterialBrush->SetResourceObject(Material);
 	}
 
-	UpdateMaterial();
-
 	ChildSlot
 	[
 		SNew(SImage)
@@ -146,96 +135,19 @@ void SMediaPlayerEditorOutput::Construct(const FArguments& InArgs, UMediaPlayer&
 
 void SMediaPlayerEditorOutput::Tick(const FGeometry& AllottedGeometry, const double InCurrentTime, const float InDeltaTime)
 {
-	UpdateMaterial();
-	UpdateSoundWave();
-}
-
-
-/* SMediaPlayerEditorOutput implementation
- *****************************************************************************/
-
-void SMediaPlayerEditorOutput::UpdateMaterial()
-{
-	UMediaTexture* Texture = MediaPlayer->GetVideoTexture();
-
-	if (Texture == nullptr)
+	if (MediaTexture != nullptr)
 	{
-		// create default texture if needed
-		if (DefaultTexture == nullptr)
-		{
-			DefaultTexture = NewObject<UMediaTexture>(GetTransientPackage(), NAME_None, RF_Transient | RF_Public);
-			
-			if (DefaultTexture != nullptr)
-			{
-				DefaultTexture->UpdateResource();
-				DefaultTexture->AddToRoot();
-			}
-		}
-
-		// set default texture as output sink
-		if (DefaultTexture != nullptr)
-		{
-			MediaPlayer->GetBasePlayer().SetVideoSink(DefaultTexture);
-		}
-
-		Texture = DefaultTexture;
-	}
-
-	// assign new texture to material
-	if (Texture != CurrentTexture)
-	{
-		TextureSampler->Texture = Texture;
-		Material->PostEditChange();
-		CurrentTexture = Texture;
-	}
-
-	// resize current texture
-	if (CurrentTexture != nullptr)
-	{
-		MaterialBrush->ImageSize.X = CurrentTexture->GetSurfaceWidth();
-		MaterialBrush->ImageSize.Y = CurrentTexture->GetSurfaceHeight();
+		MaterialBrush->ImageSize.X = MediaTexture->GetSurfaceWidth();
+		MaterialBrush->ImageSize.Y = MediaTexture->GetSurfaceHeight();
 	}
 	else
 	{
 		MaterialBrush->ImageSize = FVector2D::ZeroVector;
 	}
-}
-
-
-void SMediaPlayerEditorOutput::UpdateSoundWave()
-{
-	if (AudioComponent == nullptr)
-	{
-		return;
-	}
 
-	if (GEditor->PlayWorld == nullptr)
+	if (SoundComponent != nullptr)
 	{
-		UMediaSoundWave* SoundWave = MediaPlayer->GetSoundWave();
-
-		if (SoundWave == nullptr)
-		{
-			// set default sound wave as output sink
-			if (DefaultSoundWave != nullptr)
-			{
-				MediaPlayer->GetBasePlayer().SetAudioSink(DefaultSoundWave);
-			}
-
-			SoundWave = DefaultSoundWave;
-		}
-
-		// assign new sound wave to audio component
-		if (SoundWave != CurrentSoundWave)
-		{
-			CurrentSoundWave = SoundWave;
-			AudioComponent->SetSound(SoundWave);
-		}
-	}
-	else
-	{
-		// no audio in PIE, PIV1, Simulate
-		AudioComponent->SetSound(nullptr);
-		CurrentSoundWave = nullptr;
+		SoundComponent->UpdatePlayer();
 	}
 }
 
@@ -245,20 +157,20 @@ void SMediaPlayerEditorOutput::UpdateSoundWave()
 
 void SMediaPlayerEditorOutput::HandleMediaPlayerMediaEvent(EMediaEvent Event)
 {
-	if ((AudioComponent == nullptr) || !CurrentSoundWave.IsValid())
+	if (SoundComponent == nullptr)
 	{
 		return;
 	}
 
-	switch (Event)
+	if (Event == EMediaEvent::PlaybackSuspended)
 	{
-	case EMediaEvent::PlaybackEndReached:
-	case EMediaEvent::PlaybackSuspended:
-		AudioComponent->Stop();
-		break;
-
-	case EMediaEvent::PlaybackResumed:
-		AudioComponent->Play();
-		break;
+		SoundComponent->Stop();
+	}
+	else if (Event == EMediaEvent::PlaybackResumed)
+	{
+		if (GEditor->PlayWorld == nullptr)
+		{
+			SoundComponent->Start();
+		}
 	}
 }
diff --git a/Engine/Plugins/Media/MediaPlayerEditor/Source/MediaPlayerEditor/Private/Widgets/SMediaPlayerEditorOutput.h b/Engine/Plugins/Media/MediaPlayerEditor/Source/MediaPlayerEditor/Private/Widgets/SMediaPlayerEditorOutput.h
--- a/Engine/Plugins/Media/MediaPlayerEditor/Source/MediaPlayerEditor/Private/Widgets/SMediaPlayerEditorOutput.h
+++ b/Engine/Plugins/Media/MediaPlayerEditor/Source/MediaPlayerEditor/Private/Widgets/SMediaPlayerEditorOutput.h
@@ -7,15 +7,17 @@
 #include "Widgets/DeclarativeSyntaxSupport.h"
 #include "Widgets/SCompoundWidget.h"
 
-class UAudioComponent;
 class UMaterial;
 class UMaterialExpressionTextureSample;
 class UMediaPlayer;
-class UMediaSoundWave;
+class UMediaSoundComponent;
 class UMediaTexture;
+
 struct FSlateBrush;
+
 enum class EMediaEvent;
 
+
 /**
  * Handles content output in the viewer tab in the UMediaPlayer asset editor.
  */
@@ -51,14 +53,6 @@ public:
 
 	virtual void Tick(const FGeometry& AllottedGeometry, const double InCurrentTime, const float InDeltaTime) override;
 
-protected:
-
-	/** Update the material to reflect the media player's latest video textures. */
-	void UpdateMaterial();
-
-	/** Update the sound wave to reflect the media player's latest sound output. */
-	void UpdateSoundWave();
-
 private:
 
 	/** Callback for media events from the media player. */
@@ -66,21 +60,6 @@ private:
 
 private:
 
-	/** The audio component used to play the current sound wave. */
-	UAudioComponent* AudioComponent;
-
-	/** The media sound wave currently being played. */
-	TWeakObjectPtr<UMediaSoundWave> CurrentSoundWave;
-
-	/** The media texture currently being shown. */
-	TWeakObjectPtr<UMediaTexture> CurrentTexture;
-
-	/** The media sound wave to use if the player doesn't have one assigned. */
-	UMediaSoundWave* DefaultSoundWave;
-
-	/** The media texture to use if the player doesn't have one assigned. */
-	UMediaTexture* DefaultTexture;
-
 	/** The material that wraps the video texture for display in an SImage. */
 	UMaterial* Material;
 
@@ -88,10 +67,13 @@ private:
 	TSharedPtr<FSlateBrush> MaterialBrush;
 
 	/** The media player whose video texture is shown in this widget. */
-	UMediaPlayer* MediaPlayer;
+	TWeakObjectPtr<UMediaPlayer> MediaPlayer;
+
+	/** The media texture to render the media player's video output. */
+	UMediaTexture* MediaTexture;
 
-	/** Whether Play-In-Editor is currently active. */
-	bool PieActive;
+	/** The sound component to play the media player's audio output. */
+	UMediaSoundComponent* SoundComponent;
 
 	/** The video texture sampler in the wrapper material. */
 	UMaterialExpressionTextureSample* TextureSampler;
diff --git a/Engine/Plugins/Media/MediaPlayerEditor/Source/MediaPlayerEditor/Private/Widgets/SMediaPlayerEditorOverlay.cpp b/Engine/Plugins/Media/MediaPlayerEditor/Source/MediaPlayerEditor/Private/Widgets/SMediaPlayerEditorOverlay.cpp
--- a/Engine/Plugins/Media/MediaPlayerEditor/Source/MediaPlayerEditor/Private/Widgets/SMediaPlayerEditorOverlay.cpp
+++ b/Engine/Plugins/Media/MediaPlayerEditor/Source/MediaPlayerEditor/Private/Widgets/SMediaPlayerEditorOverlay.cpp
@@ -1,10 +1,12 @@
 // Copyright 1998-2017 Epic Games, Inc. All Rights Reserved.
 
-#include "Widgets/SMediaPlayerEditorOverlay.h"
-#include "IMediaOverlaySink.h"
-#include "Widgets/Text/SRichTextBlock.h"
+#include "SMediaPlayerEditorOverlay.h"
+
+#include "IMediaOverlaySample.h"
 #include "MediaPlayer.h"
-#include "MediaOverlays.h"
+#include "MediaPlayerFacade.h"
+#include "Widgets/Layout/SConstraintCanvas.h"
+#include "Widgets/Text/SRichTextBlock.h"
 
 
 /* SMediaPlayerEditorOverlay interface
@@ -18,16 +20,6 @@ void SMediaPlayerEditorOverlay::Construct(const FArguments& InArgs, UMediaPlayer
 	[
 		SAssignNew(Canvas, SConstraintCanvas)
 	];
-
-	for (int32 SlotIndex = 0; SlotIndex < 10; ++SlotIndex)
-	{
-		Canvas->AddSlot()
-			.AutoSize(true)
-			.Expose(Slots[SlotIndex])
-			[
-				SAssignNew(TextBlocks[SlotIndex], SRichTextBlock)
-			];
-	}
 }
 
 
@@ -36,63 +28,51 @@ void SMediaPlayerEditorOverlay::Construct(const FArguments& InArgs, UMediaPlayer
 
 void SMediaPlayerEditorOverlay::Tick(const FGeometry& AllottedGeometry, const double InCurrentTime, const float InDeltaTime)
 {
-	const float LayoutScale = AllottedGeometry.GetAccumulatedLayoutTransform().GetScale();
-	int32 InOutSlotIndex = 0;
-
-	AddOverlays(EMediaOverlayType::Caption, LayoutScale, InOutSlotIndex);
-	AddOverlays(EMediaOverlayType::Subtitle, LayoutScale, InOutSlotIndex);
-	AddOverlays(EMediaOverlayType::Text, LayoutScale, InOutSlotIndex);
-
-	while (InOutSlotIndex < 10)
-	{
-		TextBlocks[InOutSlotIndex]->SetText(FText::GetEmpty());
-		++InOutSlotIndex;
-	}
-}
-
+	Canvas->ClearChildren();
 
-/* SMediaPlayerEditorOverlay implementation
- *****************************************************************************/
+	// fetch overlay samples
+	TArray<TSharedPtr<IMediaOverlaySample, ESPMode::ThreadSafe>> OverlaySamples;
+	// @todo gmp: MF3.0 implement overlay sample caching
+//	MediaPlayer->GetPlayerFacade()->GetOverlaySamples(MediaPlayer->GetTime(), OverlaySamples);
 
-void SMediaPlayerEditorOverlay::AddOverlays(EMediaOverlayType Type, float LayoutScale, int32& InOutSlotIndex)
-{
-	if (InOutSlotIndex >= 10)
+	if (OverlaySamples.Num() == 0)
 	{
 		return;
 	}
 
-	UMediaOverlays* Overlays = MediaPlayer->GetOverlays();
-
-	if (Overlays == nullptr)
+	// create widgets for overlays
+	for (const auto& Sample : OverlaySamples)
 	{
-		return;
-	}
-
-	TArray<FMediaPlayerOverlay> OutItems;
-	Overlays->GetItems(Type, OutItems, MediaPlayer->GetTime());
+		auto RichTextBlock = SNew(SRichTextBlock)
+			.AutoWrapText(true)
+			.Justification(ETextJustify::Center)
+			.Text(Sample->GetText());
 
-	for (const FMediaPlayerOverlay& Item : OutItems)
-	{
-		SConstraintCanvas::FSlot* Slot = Slots[InOutSlotIndex];
-		TSharedPtr<SRichTextBlock> TextBlock = TextBlocks[InOutSlotIndex];
+		const TOptional<FVector2D> Position = Sample->GetPosition();
 
-		if (Item.HasPosition)
+		if (Position.IsSet())
 		{
-			Slot->Alignment(FVector2D(0.0f, 0.0f));
-			Slot->Offset(FMargin(Item.Position.X, Item.Position.Y, 0.0f, 0.0f));
+			Canvas->AddSlot()
+				.Alignment(FVector2D(0.0f, 0.0f))
+				.Anchors(FAnchors(0.0f, 0.0f, 0.0f, 0.0f))
+				.AutoSize(true)
+				.Offset(FMargin(Position->X, Position->Y, 0.0f, 0.0f))
+				[
+					RichTextBlock
+				];
 		}
 		else
 		{
-			Slot->Alignment(FVector2D(0.5f, 0.5f));
-			Slot->Offset(FMargin(0.0f));
-		}
-
-		TextBlock->SetText(Item.Text);
-		TextBlock->SlatePrepass(LayoutScale);
-
-		if (++InOutSlotIndex == 10)
-		{
-			break;
+			Canvas->AddSlot()
+				.Alignment(FVector2D(0.0f, 1.0f))
+				.Anchors(FAnchors(0.1f, 0.8f, 0.9f, 0.9f))
+				.AutoSize(true)
+				[
+					RichTextBlock
+				];
 		}
 	}
+
+	// recalculate layout
+	SlatePrepass(AllottedGeometry.GetAccumulatedLayoutTransform().GetScale());
 }
diff --git a/Engine/Plugins/Media/MediaPlayerEditor/Source/MediaPlayerEditor/Private/Widgets/SMediaPlayerEditorOverlay.h b/Engine/Plugins/Media/MediaPlayerEditor/Source/MediaPlayerEditor/Private/Widgets/SMediaPlayerEditorOverlay.h
--- a/Engine/Plugins/Media/MediaPlayerEditor/Source/MediaPlayerEditor/Private/Widgets/SMediaPlayerEditorOverlay.h
+++ b/Engine/Plugins/Media/MediaPlayerEditor/Source/MediaPlayerEditor/Private/Widgets/SMediaPlayerEditorOverlay.h
@@ -2,15 +2,13 @@
 
 #pragma once
 
-#include "CoreMinimal.h"
+#include "Templates/SharedPointer.h"
 #include "Widgets/DeclarativeSyntaxSupport.h"
 #include "Widgets/SCompoundWidget.h"
-#include "Widgets/Layout/SConstraintCanvas.h"
 
 class SConstraintCanvas;
-class SRichTextBlock;
 class UMediaPlayer;
-enum class EMediaOverlayType;
+
 
 /**
  * Draws text overlays for the UMediaPlayer asset editor.
@@ -39,15 +37,6 @@ public:
 
 	void Tick(const FGeometry& AllottedGeometry, const double InCurrentTime, const float InDeltaTime) override;
 
-protected:
-
-	/**
-	 * Add the media player's text overlays of the specified type.
-	 *
-	 * @param Type The type of overlays to add.
-	 */
-	void AddOverlays(EMediaOverlayType Type, float LayoutScale, int32& InOutSlotIndex);
-
 private:
 
 	/** The canvas to draw into. */
@@ -55,8 +44,4 @@ private:
 
 	/** The media player whose video texture is shown in this widget. */
 	UMediaPlayer* MediaPlayer;
-
-	SConstraintCanvas::FSlot* Slots[10];
-
-	TSharedPtr<SRichTextBlock> TextBlocks[10];
 };
diff --git a/Engine/Plugins/Media/MediaPlayerEditor/Source/MediaPlayerEditor/Private/Widgets/SMediaPlayerEditorPlaylist.cpp b/Engine/Plugins/Media/MediaPlayerEditor/Source/MediaPlayerEditor/Private/Widgets/SMediaPlayerEditorPlaylist.cpp
--- a/Engine/Plugins/Media/MediaPlayerEditor/Source/MediaPlayerEditor/Private/Widgets/SMediaPlayerEditorPlaylist.cpp
+++ b/Engine/Plugins/Media/MediaPlayerEditor/Source/MediaPlayerEditor/Private/Widgets/SMediaPlayerEditorPlaylist.cpp
@@ -1,10 +1,17 @@
 // Copyright 1998-2017 Epic Games, Inc. All Rights Reserved.
 
-#include "Widgets/SMediaPlayerEditorPlaylist.h"
+#include "SMediaPlayerEditorPlaylist.h"
+
+#include "Containers/Array.h"
+#include "Editor.h"
 #include "EditorStyleSet.h"
+#include "FileHelpers.h"
+#include "IMediaEventSink.h"
 #include "IMediaPlayer.h"
 #include "MediaPlayer.h"
 #include "MediaPlaylist.h"
+#include "UObject/Package.h"
+#include "Widgets/Input/SButton.h"
 #include "Widgets/Views/SListView.h"
 #include "Widgets/SMediaSourceTableRow.h"
 
@@ -43,6 +50,138 @@ void SMediaPlayerEditorPlaylist::Construct(const FArguments& InArgs, UMediaPlaye
 	[
 		SNew(SVerticalBox)
 
+		+ SVerticalBox::Slot()
+			.AutoHeight()
+			.Padding(2.0f)
+			[
+				SNew(SHorizontalBox)
+
+				+ SHorizontalBox::Slot()
+					.AutoWidth()
+					.VAlign(VAlign_Center)
+					[
+						// play list name
+						SNew(STextBlock)
+							.Text_Lambda([this]() -> FText
+							{
+								UMediaPlaylist& Playlist = MediaPlayer->GetPlaylistRef();
+								if ((Playlist.GetFlags() & RF_Transient) != 0)
+								{
+									return LOCTEXT("UnsavedPlaylistLabel", "[Unsaved play list]");
+								}
+								else
+								{
+									FString PlaylistName;
+									MediaPlayer->GetPlaylistRef().GetName(PlaylistName);
+									return FText::FromString(PlaylistName);
+								}
+							})
+					]
+
+				+ SHorizontalBox::Slot()
+					.AutoWidth()
+					.Padding(4.0f, 0.0f, 0.0f, 0.0f)
+					.VAlign(VAlign_Center)
+					[
+						// browse button
+						SNew(SButton)
+							.ButtonStyle(FEditorStyle::Get(), "HoverHintOnly")
+							.ContentPadding(4.0f)
+							.ForegroundColor(FSlateColor::UseForeground())
+							.IsEnabled_Lambda([this]() -> bool {
+								return ((MediaPlayer->GetPlaylistRef().GetFlags() & RF_Transient) == 0);
+							})
+							.OnClicked_Lambda([this]() -> FReply
+							{
+								TArray<UObject*> AssetsToSync;
+								AssetsToSync.Add(MediaPlayer->GetPlaylist());
+								GEditor->SyncBrowserToObjects(AssetsToSync);
+								return FReply::Handled();
+							})
+							.ToolTipText(LOCTEXT("FindPlaylistButtonToolTip", "Find this playlist in the Content Browser"))
+							[
+								SNew(SImage)
+									.ColorAndOpacity(FSlateColor::UseForeground())
+									.Image(FEditorStyle::GetBrush("PropertyWindow.Button_Browse"))
+							]
+					]
+
+				+ SHorizontalBox::Slot()
+					.FillWidth(1.0f)
+					.HAlign(HAlign_Left)
+					.VAlign(VAlign_Center)
+					[
+						// save button
+						SNew(SButton)
+							.ButtonStyle(FEditorStyle::Get(), "HoverHintOnly")
+							.ContentPadding(4.0f)
+							.ForegroundColor(FSlateColor::UseForeground())
+							.IsEnabled_Lambda([this]() -> bool {
+								UMediaPlaylist& Playlist = MediaPlayer->GetPlaylistRef();
+								UPackage* Package = Playlist.GetOutermost();
+								return ((Playlist.GetFlags() & RF_Transient) != 0) || ((Package != nullptr) && Package->IsDirty());
+							})
+							.OnClicked(this, &SMediaPlayerEditorPlaylist::HandleSavePlaylistButtonClicked)
+							.ToolTipText(LOCTEXT("SavePlaylistButtonToolTip", "Save this playlist"))
+							[
+								SNew(SImage)
+									.ColorAndOpacity(FSlateColor::UseForeground())
+									.Image_Lambda([this]() -> const FSlateBrush*
+									{
+										return FEditorStyle::GetBrush(
+											((MediaPlayer->GetPlaylistRef().GetFlags() & RF_Transient) != 0)
+												? "LevelEditor.Save"
+												: "LevelEditor.SaveAs"
+										);
+									})
+							]
+					]
+
+				+ SHorizontalBox::Slot()
+					.AutoWidth()
+					.Padding(2.0f, 0.0f, 0.0f, 0.0f)
+					.VAlign(VAlign_Center)
+					[
+						// previous item button
+						SNew(SButton)
+							.IsEnabled_Lambda([this]() -> bool {
+								return (MediaPlayer->GetPlaylistRef().Loop || (MediaPlayer->GetPlaylistIndex() > 0));
+							})
+							.OnClicked_Lambda([this]() -> FReply
+							{
+								MediaPlayer->Previous();
+								return FReply::Handled();
+							})
+							.ToolTipText(LOCTEXT("PreviousPlaylistItemButtonToolTip", "Jump to the previous item in the playlist"))
+							[
+								SNew(SImage)
+									.Image(FEditorStyle::GetBrush("ContentBrowser.HistoryBack"))
+							]
+					]
+
+				+ SHorizontalBox::Slot()
+					.AutoWidth()
+					.Padding(2.0f, 0.0f, 0.0f, 0.0f)
+					.VAlign(VAlign_Center)
+					[
+						// next item button
+						SNew(SButton)
+							.IsEnabled_Lambda([this]() -> bool {
+								return (MediaPlayer->GetPlaylistRef().Loop || (MediaPlayer->GetPlaylistIndex() + 1 < MediaPlayer->GetPlaylistRef().Num()));
+							})
+							.OnClicked_Lambda([this]() -> FReply
+							{
+								MediaPlayer->Next();
+								return FReply::Handled();
+							})
+							.ToolTipText(LOCTEXT("NextPlaylistItemButtonToolTip", "Jump to the next item in the playlist"))
+							[
+								SNew(SImage)
+									.Image(FEditorStyle::GetBrush("ContentBrowser.HistoryForward"))
+							]
+					]
+			]
+
 		+ SVerticalBox::Slot()
 			.FillHeight(1.0f)
 			[
@@ -50,13 +189,25 @@ void SMediaPlayerEditorPlaylist::Construct(const FArguments& InArgs, UMediaPlaye
 					.BorderImage(FEditorStyle::GetBrush("ToolPanel.GroupBorder"))
 					.Padding(0.0f)
 					[
-						// message list
+						// media source list
 						SAssignNew(MediaSourceListView, SListView<TSharedPtr<FMediaSourceTableEntry>>)
 							.ItemHeight(24.0f)
 							.ListItemsSource(&MediaSourceList)
 							.SelectionMode(ESelectionMode::Single)
-							.OnGenerateRow(this, &SMediaPlayerEditorPlaylist::HandleMediaSourceListGenerateRow)
-							.OnMouseButtonDoubleClick(this, &SMediaPlayerEditorPlaylist::HandleMediaSourceListDoubleClick)
+							.OnGenerateRow_Lambda([this](TSharedPtr<FMediaSourceTableEntry> Entry, const TSharedRef<STableViewBase>& OwnerTable) -> TSharedRef<ITableRow>
+							{
+								return SNew(SMediaSourceTableRow, OwnerTable)
+									.Entry(Entry)
+									.Opened_Lambda([=]() { return MediaPlayer->GetPlaylistIndex() == Entry->Index; })
+									.Style(Style);
+							})
+							.OnMouseButtonDoubleClick_Lambda([this](TSharedPtr<FMediaSourceTableEntry> InItem)
+							{
+								if (InItem.IsValid())
+								{
+									MediaPlayer->OpenPlaylistIndex(&MediaPlayer->GetPlaylistRef(), InItem->Index);
+								}
+							})
 							.HeaderRow
 							(
 								SNew(SHeaderRow)
@@ -95,14 +246,11 @@ void SMediaPlayerEditorPlaylist::ReloadMediaSourceList()
 {
 	MediaSourceList.Reset();
 
-	UMediaPlaylist* Playlist = MediaPlayer->GetPlaylist();
+	UMediaPlaylist& Playlist = MediaPlayer->GetPlaylistRef();
 
-	if (Playlist != nullptr)
+	for (int32 EntryIndex = 0; EntryIndex < Playlist.Num(); ++EntryIndex)
 	{
-		for (int32 EntryIndex = 0; EntryIndex < Playlist->Num(); ++EntryIndex)
-		{
-			MediaSourceList.Add(MakeShareable(new FMediaSourceTableEntry(EntryIndex, Playlist->Get(EntryIndex))));
-		}
+		MediaSourceList.Add(MakeShareable(new FMediaSourceTableEntry(EntryIndex, Playlist.Get(EntryIndex))));
 	}
 
 	MediaSourceListView->RequestListRefresh();
@@ -114,7 +262,7 @@ void SMediaPlayerEditorPlaylist::ReloadMediaSourceList()
 
 void SMediaPlayerEditorPlaylist::HandleCoreObjectPropertyChanged(UObject* Object, struct FPropertyChangedEvent& ChangedEvent)
 {
-	if ((Object != nullptr) && (Object == MediaPlayer->GetPlaylist()))
+	if ((Object != nullptr) && (Object == &MediaPlayer->GetPlaylistRef()))
 	{
 		ReloadMediaSourceList();
 	}
@@ -130,21 +278,59 @@ void SMediaPlayerEditorPlaylist::HandleMediaPlayerMediaEvent(EMediaEvent Event)
 }
 
 
-void SMediaPlayerEditorPlaylist::HandleMediaSourceListDoubleClick(TSharedPtr<FMediaSourceTableEntry> InItem)
+FReply SMediaPlayerEditorPlaylist::HandleSavePlaylistButtonClicked()
 {
-	if (InItem.IsValid())
+	UMediaPlaylist& Playlist = MediaPlayer->GetPlaylistRef();
+
+	// first save any transient media sources
+	for (int32 MediaSourceIndex = 0; MediaSourceIndex < Playlist.Num(); ++MediaSourceIndex)
 	{
-		MediaPlayer->OpenPlaylistIndex(MediaPlayer->GetPlaylist(), InItem->Index);
+		UMediaSource* MediaSource = Playlist.Get(MediaSourceIndex);
+
+		if ((MediaSource != nullptr) && MediaSource->HasAnyFlags(RF_Transient))
+		{
+			TArray<UObject*> MediaSourcesToSave;
+			TArray<UObject*> SavedMediaSources;
+
+			MediaSourcesToSave.Add(MediaSource);
+			FEditorFileUtils::SaveAssetsAs(MediaSourcesToSave, SavedMediaSources);
+
+			if ((SavedMediaSources.Num() != 1) || (SavedMediaSources[0] == MediaSourcesToSave[0]))
+			{
+				return FReply::Handled(); // user canceled
+			}
+
+			Playlist.Replace(MediaSourceIndex, CastChecked<UMediaSource>(SavedMediaSources[0]));
+		}
 	}
-}
 
+	// then save play list itself
+	if ((Playlist.GetFlags() & RF_Transient) == 0)
+	{
+		// save existing play list asset
+		TArray<UPackage*> PackagesToSave;
+
+		PackagesToSave.Add(Playlist.GetOutermost());
+		FEditorFileUtils::PromptForCheckoutAndSave(PackagesToSave, false /*bCheckDirty*/, false /*bPromptToSave*/);
+	}
+	else
+	{
+		// create & save new play list asset
+		TArray<UObject*> PlaylistsToSave;
+		TArray<UObject*> SavedPlaylists;
 
-TSharedRef<ITableRow> SMediaPlayerEditorPlaylist::HandleMediaSourceListGenerateRow(TSharedPtr<FMediaSourceTableEntry> Entry, const TSharedRef<STableViewBase>& OwnerTable)
-{
-	return SNew(SMediaSourceTableRow, OwnerTable)
-		.Entry(Entry)
-		.Opened_Lambda([=]() { return MediaPlayer->GetPlaylistIndex() == Entry->Index; })
-		.Style(Style);
+		PlaylistsToSave.Add(&Playlist);
+		FEditorFileUtils::SaveAssetsAs(PlaylistsToSave, SavedPlaylists);
+
+		if ((SavedPlaylists.Num() != 1) || (SavedPlaylists[0] == &Playlist))
+		{
+			return FReply::Handled(); // user canceled
+		}
+
+		MediaPlayer->OpenPlaylistIndex(CastChecked<UMediaPlaylist>(SavedPlaylists[0]), MediaPlayer->GetPlaylistIndex());
+	}
+
+	return FReply::Handled();
 }
 
 
