Copying //UE4/Dev-Sequencer to //UE4/Dev-Main (Source: //UE4/Dev-Sequencer @ 3946692)

diff --git a/Engine/Plugins/FX/Niagara/Source/NiagaraEditor/Private/ViewModels/NiagaraSystemViewModel.cpp b/Engine/Plugins/FX/Niagara/Source/NiagaraEditor/Private/ViewModels/NiagaraSystemViewModel.cpp
--- a/Engine/Plugins/FX/Niagara/Source/NiagaraEditor/Private/ViewModels/NiagaraSystemViewModel.cpp
+++ b/Engine/Plugins/FX/Niagara/Source/NiagaraEditor/Private/ViewModels/NiagaraSystemViewModel.cpp
@@ -31,6 +31,7 @@
 #include "Math/NumericLimits.h"
 #include "NiagaraComponent.h"
 #include "MovieSceneFolder.h"
+#include "Channels/MovieSceneChannelProxy.h"
 
 DECLARE_CYCLE_STAT(TEXT("Niagara - SystemViewModel - CompileSystem"), STAT_NiagaraEditor_SystemViewModel_CompileSystem, STATGROUP_NiagaraEditor);
 
@@ -687,11 +688,9 @@ void FNiagaraSystemViewModel::RefreshSequencerTracks()
 	}
 
 	// Expand the view range to show all emitters.
-	TRange<float> CurrentViewRange = NiagaraSequence->GetMovieScene()->GetEditorData().ViewRange;
-	float NewMinViewRange = FMath::Min(MinEmitterTime, CurrentViewRange.GetLowerBoundValue());
-	float NewMaxViewRange = FMath::Max(MaxEmitterTime, CurrentViewRange.GetUpperBoundValue());
-
-	NiagaraSequence->GetMovieScene()->GetEditorData().ViewRange = TRange<float>(NewMinViewRange, NewMaxViewRange);
+	FMovieSceneEditorData& EditorData = NiagaraSequence->GetMovieScene()->GetEditorData();
+	EditorData.ViewStart = FMath::Min<double>(MinEmitterTime, EditorData.ViewStart);
+	EditorData.ViewEnd   = FMath::Max<double>(MaxEmitterTime, EditorData.ViewEnd);
 
 	Sequencer->NotifyMovieSceneDataChanged(EMovieSceneDataChangeType::MovieSceneStructureItemsChanged);
 	Sequencer->SetGlobalTime(0);
@@ -745,24 +744,29 @@ void FNiagaraSystemViewModel::RefreshSequencerTrack(UMovieSceneNiagaraEmitterTra
 				EndTime = EmitterHandleViewModel->GetEmitterViewModel()->GetEndTime();
 			}
 
+			FFrameRate FrameResolution = EmitterSection->GetTypedOuter<UMovieScene>()->GetFrameResolution();
+
 			EmitterSection->SetEmitterHandle(EmitterHandleViewModel.ToSharedRef());
 			EmitterSection->SetIsActive(EmitterHandleViewModel->GetIsEnabled());
-			EmitterSection->SetStartTime(StartTime);
-			EmitterSection->SetEndTime(EndTime);
-			EmitterSection->SetIsInfinite(bIsInfinite);
+			if (bIsInfinite)
+			{
+				EmitterSection->SetRange(TRange<FFrameNumber>((StartTime * FrameResolution).FloorToFrame(), (EndTime * FrameResolution).FloorToFrame()));
+			}
+			else
+			{
+				EmitterSection->SetRange(TRange<FFrameNumber>::All());
+			}
 
-// 			TSharedPtr<FBurstCurve> BurstCurve = EmitterSection->GetBurstCurve();
-// 			if (BurstCurve.IsValid())
+			TMovieSceneChannel<FMovieSceneBurstKey> Channel = EmitterSection->GetChannelProxy().GetChannel<FMovieSceneNiagaraEmitterChannel>(0)->GetInterface();
+
+// 			Channel.Reset();
+// 			for (const FNiagaraEmitterBurst& Burst : Emitter->Bursts)
 // 			{
-// 				BurstCurve->Reset();
-// 				for (const FNiagaraEmitterBurst& Burst : Emitter->Bursts)
-// 				{
-// 					FMovieSceneBurstKey BurstKey;
-// 					BurstKey.TimeRange = Burst.TimeRange;
-// 					BurstKey.SpawnMinimum = Burst.SpawnMinimum;
-// 					BurstKey.SpawnMaximum = Burst.SpawnMaximum;
-// 					BurstCurve->AddKeyValue(Burst.Time, BurstKey);
-// 				}
+// 				FMovieSceneBurstKey BurstKey;
+// 				BurstKey.TimeRange = Burst.TimeRange;
+// 				BurstKey.SpawnMinimum = Burst.SpawnMinimum;
+// 				BurstKey.SpawnMaximum = Burst.SpawnMaximum;
+// 				Channel.AddKey(Burst.Time, BurstKey);
 // 			}
 		}
 	}
@@ -772,15 +776,20 @@ void FNiagaraSystemViewModel::SetupSequencer()
 {
 	NiagaraSequence = NewObject<UNiagaraSequence>(GetTransientPackage());
 	UMovieScene* MovieScene = NewObject<UMovieScene>(NiagaraSequence, FName("Niagara System MovieScene"), RF_Transactional);
-	MovieScene->SetPlaybackRange(SequencerDefaultPlaybackRange.GetLowerBoundValue(), SequencerDefaultPlaybackRange.GetUpperBoundValue());
-	MovieScene->GetEditorData().ViewRange = SequencerDefaultViewRange;
-	MovieScene->GetEditorData().WorkingRange = SequencerDefaultPlaybackRange;
 
 	NiagaraSequence->Initialize(this, MovieScene);
 
+	FFrameTime StartTime = SequencerDefaultPlaybackRange.GetLowerBoundValue() * MovieScene->GetFrameResolution();
+	int32      Duration  = (SequencerDefaultPlaybackRange.Size<float>() * MovieScene->GetFrameResolution()).FrameNumber.Value;
+
+	MovieScene->SetPlaybackRange(StartTime.RoundToFrame(), Duration);
+
+	FMovieSceneEditorData& EditorData = NiagaraSequence->GetMovieScene()->GetEditorData();
+	EditorData.ViewStart = EditorData.WorkStart = SequencerDefaultViewRange.GetLowerBoundValue();
+	EditorData.ViewEnd   = EditorData.WorkEnd   = SequencerDefaultViewRange.GetUpperBoundValue();
+
 	FSequencerViewParams ViewParams(TEXT("NiagaraSequencerSettings"));
 	{
-		ViewParams.InitialScrubPosition = 0;
 		ViewParams.UniqueName = "NiagaraSequenceEditor";
 		ViewParams.OnGetAddMenuContent = OnGetSequencerAddMenuContent;
 	}
@@ -807,7 +816,7 @@ void FNiagaraSystemViewModel::SetupSequencer()
 
 void FNiagaraSystemViewModel::ResetSystem()
 {
-	Sequencer->SetGlobalTime(0.0f);
+	Sequencer->SetGlobalTime(0);
 	Sequencer->SetPlaybackStatus(EMovieScenePlayerStatus::Playing);
 	
 	for (TObjectIterator<UNiagaraComponent> ComponentIt; ComponentIt; ++ComponentIt)
@@ -838,7 +847,7 @@ void FNiagaraSystemViewModel::ReInitializeSystemInstances()
 {
 	if (Sequencer.IsValid() && Sequencer->GetPlaybackStatus() == EMovieScenePlayerStatus::Playing)
 	{
-		Sequencer->SetGlobalTime(0.0f);
+		Sequencer->SetGlobalTime(0);
 	}
 
 	for (TObjectIterator<UNiagaraComponent> ComponentIt; ComponentIt; ++ComponentIt)
@@ -1075,39 +1084,37 @@ void FNiagaraSystemViewModel::SequencerDataChanged(EMovieSceneDataChangeType Dat
 			EmitterHandleViewModel->SetIsEnabled(EmitterSection->IsActive());
 
 			TSharedPtr<FNiagaraEmitterViewModel> EmitterViewModel = EmitterTrack->GetEmitterHandle()->GetEmitterViewModel();
-			if (EmitterSection->IsInfinite() == false)
-			{
-				EmitterViewModel->SetStartTime(EmitterSection->GetStartTime());
-				EmitterViewModel->SetEndTime(EmitterSection->GetEndTime());
-			}
 
-			TSharedPtr<FBurstCurve> BurstCurve = EmitterSection->GetBurstCurve();
-			if (BurstCurve.IsValid())
+			// @todo: sequencer-timecode: these don't seem to actually do anything???
+			// if (EmitterSection->GetRange() != TRange<FFrameNumber>::All())
+			// {
+			// 	EmitterViewModel->SetStartTime(EmitterSection->GetStartTime());
+			// 	EmitterViewModel->SetEndTime(EmitterSection->GetEndTime());
+			// }
+
+			const FMovieSceneNiagaraEmitterChannel* Channel = EmitterSection->GetChannelProxy().GetChannel<FMovieSceneNiagaraEmitterChannel>(0);
+
+			TArrayView<const FFrameNumber>        Times  = Channel->GetTimes();
+			TArrayView<const FMovieSceneBurstKey> Values = Channel->GetValues();
+
+// 			UNiagaraEmitter* Emitter = EmitterViewModel->GetEmitter();
+// 			if (Emitter)
+// 			{
+// 				Emitter->Bursts.Empty();
+// 			}
+
+			FFrameRate FrameResolution = EmitterSection->GetTypedOuter<UMovieScene>()->GetFrameResolution();
+			for (int32 Index = 0; Index < Times.Num(); ++Index)
 			{
-				UNiagaraEmitter* Emitter = EmitterViewModel->GetEmitter();
+				FNiagaraEmitterBurst Burst;
+				Burst.Time = Times[Index] / FrameResolution;
+				Burst.TimeRange    = Values[Index].TimeRange / FrameResolution;
+				Burst.SpawnMinimum = Values[Index].SpawnMinimum;
+				Burst.SpawnMaximum = Values[Index].SpawnMaximum;
 // 				if (Emitter)
 // 				{
-// 					Emitter->Bursts.Empty();
+// 					Emitter->Bursts.Add(Burst);
 // 				}
-				TArray<FKeyHandle> KeyHandles;
-				for (TKeyTimeIterator<float> KeyIterator = BurstCurve->IterateKeys(); KeyIterator; ++KeyIterator)
-				{
-					KeyHandles.Add(KeyIterator.GetKeyHandle());
-				}
-
-				for (const FKeyHandle& KeyHandle : KeyHandles)
-				{
-					auto Key = BurstCurve->GetKey(KeyHandle).GetValue();
-					FNiagaraEmitterBurst Burst;
-					Burst.Time = Key.Time;
-					Burst.TimeRange = Key.Value.TimeRange;
-					Burst.SpawnMinimum = Key.Value.SpawnMinimum;
-					Burst.SpawnMaximum = Key.Value.SpawnMaximum;
-// 					if (Emitter)
-// 					{
-// 						Emitter->Bursts.Add(Burst);
-// 					}
-				}
 			}
 		}
 	}
@@ -1169,7 +1176,7 @@ void FNiagaraSystemViewModel::SequencerTimeChanged()
 		return;
 	}
 	EMovieScenePlayerStatus::Type CurrentStatus = Sequencer->GetPlaybackStatus();
-	float CurrentSequencerTime = Sequencer->GetGlobalTime();
+	float CurrentSequencerTime = Sequencer->GetGlobalTime().AsSeconds();
 	if (SystemInstance != nullptr)
 	{
 		// Avoid reentrancy if we're setting the time directly.
diff --git a/Engine/Plugins/FX/Niagara/Source/NiagaraEditor/Private/Widgets/SNiagaraCurveEditor.cpp b/Engine/Plugins/FX/Niagara/Source/NiagaraEditor/Private/Widgets/SNiagaraCurveEditor.cpp
--- a/Engine/Plugins/FX/Niagara/Source/NiagaraEditor/Private/Widgets/SNiagaraCurveEditor.cpp
+++ b/Engine/Plugins/FX/Niagara/Source/NiagaraEditor/Private/Widgets/SNiagaraCurveEditor.cpp
@@ -5,7 +5,7 @@
 #include "SCurveEditor.h"
 #include "Widgets/SOverlay.h"
 #include "Framework/MultiBox/MultiBoxBuilder.h"
-#include "RichCurveEditorCommands.h"
+#include "CurveEditorCommands.h"
 #include "Widgets/Input/SNumericDropDown.h"
 #include "SCurveEditor.h"
 #include "Widgets/SOverlay.h"
@@ -81,35 +81,35 @@ TSharedRef<SWidget> SNiagaraCurveEditor::ConstructToolBar(TSharedPtr<FUICommandL
 
 	ToolBarBuilder.BeginSection("Snap");
 	{
-		ToolBarBuilder.AddToolBarButton(FRichCurveEditorCommands::Get().ToggleInputSnapping);
+		ToolBarBuilder.AddToolBarButton(FCurveEditorCommands::Get().ToggleInputSnapping);
 		ToolBarBuilder.AddWidget(InputSnapWidget);
-		ToolBarBuilder.AddToolBarButton(FRichCurveEditorCommands::Get().ToggleOutputSnapping);
+		ToolBarBuilder.AddToolBarButton(FCurveEditorCommands::Get().ToggleOutputSnapping);
 		ToolBarBuilder.AddWidget(OutputSnapWidget);
 	}
 	ToolBarBuilder.EndSection();
 
 	ToolBarBuilder.BeginSection("Curve");
 	{
-		ToolBarBuilder.AddToolBarButton(FRichCurveEditorCommands::Get().ZoomToFitHorizontal);
-		ToolBarBuilder.AddToolBarButton(FRichCurveEditorCommands::Get().ZoomToFitVertical);
-		ToolBarBuilder.AddToolBarButton(FRichCurveEditorCommands::Get().ZoomToFit);
+		ToolBarBuilder.AddToolBarButton(FCurveEditorCommands::Get().ZoomToFitHorizontal);
+		ToolBarBuilder.AddToolBarButton(FCurveEditorCommands::Get().ZoomToFitVertical);
+		ToolBarBuilder.AddToolBarButton(FCurveEditorCommands::Get().ZoomToFit);
 	}
 	ToolBarBuilder.EndSection();
 
 	ToolBarBuilder.BeginSection("Interpolation");
 	{
-		ToolBarBuilder.AddToolBarButton(FRichCurveEditorCommands::Get().InterpolationCubicAuto);
-		ToolBarBuilder.AddToolBarButton(FRichCurveEditorCommands::Get().InterpolationCubicUser);
-		ToolBarBuilder.AddToolBarButton(FRichCurveEditorCommands::Get().InterpolationCubicBreak);
-		ToolBarBuilder.AddToolBarButton(FRichCurveEditorCommands::Get().InterpolationLinear);
-		ToolBarBuilder.AddToolBarButton(FRichCurveEditorCommands::Get().InterpolationConstant);
+		ToolBarBuilder.AddToolBarButton(FCurveEditorCommands::Get().InterpolationCubicAuto);
+		ToolBarBuilder.AddToolBarButton(FCurveEditorCommands::Get().InterpolationCubicUser);
+		ToolBarBuilder.AddToolBarButton(FCurveEditorCommands::Get().InterpolationCubicBreak);
+		ToolBarBuilder.AddToolBarButton(FCurveEditorCommands::Get().InterpolationLinear);
+		ToolBarBuilder.AddToolBarButton(FCurveEditorCommands::Get().InterpolationConstant);
 	}
 	ToolBarBuilder.EndSection();
 
 	ToolBarBuilder.BeginSection("Tangents");
 	{
-		ToolBarBuilder.AddToolBarButton(FRichCurveEditorCommands::Get().FlattenTangents);
-		ToolBarBuilder.AddToolBarButton(FRichCurveEditorCommands::Get().StraightenTangents);
+		ToolBarBuilder.AddToolBarButton(FCurveEditorCommands::Get().FlattenTangents);
+		ToolBarBuilder.AddToolBarButton(FCurveEditorCommands::Get().StraightenTangents);
 	}
 	ToolBarBuilder.EndSection();
 
@@ -128,26 +128,18 @@ TSharedRef<SWidget> SNiagaraCurveEditor::MakeCurveEditorViewOptionsMenu(TSharedP
 {
 	FMenuBuilder MenuBuilder(true, CurveEditorCommandList);
 
-	MenuBuilder.BeginSection("CurveVisibility", LOCTEXT("CurveEditorMenuCurveVisibilityHeader", "Curve Visibility"));
-	{
-		MenuBuilder.AddMenuEntry(FRichCurveEditorCommands::Get().SetAllCurveVisibility);
-		MenuBuilder.AddMenuEntry(FRichCurveEditorCommands::Get().SetSelectedCurveVisibility);
-		MenuBuilder.AddMenuEntry(FRichCurveEditorCommands::Get().SetAnimatedCurveVisibility);
-	}
-	MenuBuilder.EndSection();
-
 	MenuBuilder.BeginSection("TangentVisibility", LOCTEXT("CurveEditorMenuTangentVisibilityHeader", "Tangent Visibility"));
 	{
-		MenuBuilder.AddMenuEntry(FRichCurveEditorCommands::Get().SetAllTangentsVisibility);
-		MenuBuilder.AddMenuEntry(FRichCurveEditorCommands::Get().SetSelectedKeysTangentVisibility);
-		MenuBuilder.AddMenuEntry(FRichCurveEditorCommands::Get().SetNoTangentsVisibility);
+		MenuBuilder.AddMenuEntry(FCurveEditorCommands::Get().SetAllTangentsVisibility);
+		MenuBuilder.AddMenuEntry(FCurveEditorCommands::Get().SetSelectedKeysTangentVisibility);
+		MenuBuilder.AddMenuEntry(FCurveEditorCommands::Get().SetNoTangentsVisibility);
 	}
 	MenuBuilder.EndSection();
 
 	MenuBuilder.AddMenuSeparator();
 
-	MenuBuilder.AddMenuEntry(FRichCurveEditorCommands::Get().ToggleAutoFrameCurveEditor);
-	MenuBuilder.AddMenuEntry(FRichCurveEditorCommands::Get().ToggleShowCurveEditorCurveToolTips);
+	MenuBuilder.AddMenuEntry(FCurveEditorCommands::Get().ToggleAutoFrameCurveEditor);
+	MenuBuilder.AddMenuEntry(FCurveEditorCommands::Get().ToggleShowCurveEditorCurveToolTips);
 
 	return MenuBuilder.MakeWidget();
 }
@@ -159,11 +151,11 @@ TSharedRef<SWidget> SNiagaraCurveEditor::MakeCurveEditorCurveOptionsMenu(TShared
 		{
 			MenuBuilder.BeginSection("Pre-Infinity Extrapolation", LOCTEXT("CurveEditorMenuPreInfinityExtrapHeader", "Extrapolation"));
 			{
-				MenuBuilder.AddMenuEntry(FRichCurveEditorCommands::Get().SetPreInfinityExtrapCycle);
-				MenuBuilder.AddMenuEntry(FRichCurveEditorCommands::Get().SetPreInfinityExtrapCycleWithOffset);
-				MenuBuilder.AddMenuEntry(FRichCurveEditorCommands::Get().SetPreInfinityExtrapOscillate);
-				MenuBuilder.AddMenuEntry(FRichCurveEditorCommands::Get().SetPreInfinityExtrapLinear);
-				MenuBuilder.AddMenuEntry(FRichCurveEditorCommands::Get().SetPreInfinityExtrapConstant);
+				MenuBuilder.AddMenuEntry(FCurveEditorCommands::Get().SetPreInfinityExtrapCycle);
+				MenuBuilder.AddMenuEntry(FCurveEditorCommands::Get().SetPreInfinityExtrapCycleWithOffset);
+				MenuBuilder.AddMenuEntry(FCurveEditorCommands::Get().SetPreInfinityExtrapOscillate);
+				MenuBuilder.AddMenuEntry(FCurveEditorCommands::Get().SetPreInfinityExtrapLinear);
+				MenuBuilder.AddMenuEntry(FCurveEditorCommands::Get().SetPreInfinityExtrapConstant);
 			}
 			MenuBuilder.EndSection();
 		}
@@ -172,11 +164,11 @@ TSharedRef<SWidget> SNiagaraCurveEditor::MakeCurveEditorCurveOptionsMenu(TShared
 		{
 			MenuBuilder.BeginSection("Post-Infinity Extrapolation", LOCTEXT("CurveEditorMenuPostInfinityExtrapHeader", "Extrapolation"));
 			{
-				MenuBuilder.AddMenuEntry(FRichCurveEditorCommands::Get().SetPostInfinityExtrapCycle);
-				MenuBuilder.AddMenuEntry(FRichCurveEditorCommands::Get().SetPostInfinityExtrapCycleWithOffset);
-				MenuBuilder.AddMenuEntry(FRichCurveEditorCommands::Get().SetPostInfinityExtrapOscillate);
-				MenuBuilder.AddMenuEntry(FRichCurveEditorCommands::Get().SetPostInfinityExtrapLinear);
-				MenuBuilder.AddMenuEntry(FRichCurveEditorCommands::Get().SetPostInfinityExtrapConstant);
+				MenuBuilder.AddMenuEntry(FCurveEditorCommands::Get().SetPostInfinityExtrapCycle);
+				MenuBuilder.AddMenuEntry(FCurveEditorCommands::Get().SetPostInfinityExtrapCycleWithOffset);
+				MenuBuilder.AddMenuEntry(FCurveEditorCommands::Get().SetPostInfinityExtrapOscillate);
+				MenuBuilder.AddMenuEntry(FCurveEditorCommands::Get().SetPostInfinityExtrapLinear);
+				MenuBuilder.AddMenuEntry(FCurveEditorCommands::Get().SetPostInfinityExtrapConstant);
 			}
 			MenuBuilder.EndSection();
 		}
@@ -184,8 +176,8 @@ TSharedRef<SWidget> SNiagaraCurveEditor::MakeCurveEditorCurveOptionsMenu(TShared
 
 	FMenuBuilder MenuBuilder(true, CurveEditorCommandList);
 
-	MenuBuilder.AddMenuEntry(FRichCurveEditorCommands::Get().BakeCurve);
-	MenuBuilder.AddMenuEntry(FRichCurveEditorCommands::Get().ReduceCurve);
+	MenuBuilder.AddMenuEntry(FCurveEditorCommands::Get().BakeCurve);
+	MenuBuilder.AddMenuEntry(FCurveEditorCommands::Get().ReduceCurve);
 
 	MenuBuilder.AddSubMenu(
 		LOCTEXT("PreInfinitySubMenu", "Pre-Infinity"),
diff --git a/Engine/Plugins/FX/Niagara/Source/NiagaraEditor/Private/Widgets/SNiagaraSpreadsheetView.cpp b/Engine/Plugins/FX/Niagara/Source/NiagaraEditor/Private/Widgets/SNiagaraSpreadsheetView.cpp
--- a/Engine/Plugins/FX/Niagara/Source/NiagaraEditor/Private/Widgets/SNiagaraSpreadsheetView.cpp
+++ b/Engine/Plugins/FX/Niagara/Source/NiagaraEditor/Private/Widgets/SNiagaraSpreadsheetView.cpp
@@ -629,7 +629,7 @@ void SNiagaraSpreadsheetView::HandleTimeChange()
 					}
 				}
 
-				float LocalCaptureTime = SystemViewModel->GetSequencer()->GetLocalTime();
+				FQualifiedFrameTime LocalCaptureTime = SystemViewModel->GetSequencer()->GetLocalTime();
 
 				if (FoundScript && FoundScript->GetDebuggerInfo().bRequestDebugFrame == false && CaptureData[i].LastReadWriteId != FoundScript->GetDebuggerInfo().DebugFrameLastWriteId)
 				{
@@ -638,7 +638,7 @@ void SNiagaraSpreadsheetView::HandleTimeChange()
 					CaptureData[i].InputParams = FoundScript->GetDebuggerInfo().DebugParameters;
 					CaptureData[i].DataSet->Tick(ENiagaraSimTarget::CPUSim); // Force a buffer swap, from here out we read from PrevData.
 
-					CaptureData[i].LastCaptureTime = LocalCaptureTime;
+					CaptureData[i].LastCaptureTime = LocalCaptureTime.AsSeconds();
 					ensure(CaptureData[i].LastCaptureTime == CaptureData[i].TargetCaptureTime);
 					CaptureData[i].LastCaptureHandleId = SelectedEmitterHandles[0]->GetId();
 
@@ -685,7 +685,7 @@ bool SNiagaraSpreadsheetView::IsPausedAtRightTimeOnRightHandle() const
 	if (SelectedEmitterHandles.Num() == 1)
 	{
 		return SystemViewModel->GetSequencer()->GetPlaybackStatus() == EMovieScenePlayerStatus::Stopped &&
-			CaptureData[(int32)TabState].LastCaptureTime == SystemViewModel->GetSequencer()->GetLocalTime() &&
+			CaptureData[(int32)TabState].LastCaptureTime == SystemViewModel->GetSequencer()->GetLocalTime().AsSeconds() &&
 			CaptureData[(int32)TabState].LastCaptureHandleId == SelectedEmitterHandles[0]->GetId();
 	}
 	return false;
@@ -989,8 +989,9 @@ void SNiagaraSpreadsheetView::ResetColumns(EUITab Tab)
 
 FReply SNiagaraSpreadsheetView::OnCaptureRequestPressed()
 {
-	float LocalTime = SystemViewModel->GetSequencer()->GetLocalTime();
-	float SnapInterval = SystemViewModel->GetSequencer()->GetSequencerSettings()->GetTimeSnapInterval();
+	FFrameRate FrameResolution = SystemViewModel->GetSequencer()->GetFocusedFrameResolution();
+	float LocalTime = SystemViewModel->GetSequencer()->GetLocalTime().AsSeconds();
+	float SnapInterval = SystemViewModel->GetSequencer()->GetFocusedPlayRate().AsInterval();
 	float TargetCaptureTime = LocalTime + SnapInterval;
 
 	TArray<TSharedRef<FNiagaraEmitterHandleViewModel>> SelectedEmitterHandles;
@@ -1021,18 +1022,18 @@ FReply SNiagaraSpreadsheetView::OnCaptureRequestPressed()
 			FoundScript->GetDebuggerInfo().bRequestDebugFrame = true;
 			CaptureData[(int32)i].bAwaitingFrame = true;
 
-			CaptureData[(int32)i].TargetCaptureTime = TargetCaptureTime;			
+			CaptureData[(int32)i].TargetCaptureTime = TargetCaptureTime;
 		}
 	}
 
 	if (SystemViewModel->GetSequencer()->GetPlaybackStatus() == EMovieScenePlayerStatus::Stopped)
 	{
-		SystemViewModel->GetSequencer()->SetLocalTime(TargetCaptureTime, STM_None);
+		SystemViewModel->GetSequencer()->SetLocalTime(TargetCaptureTime * FrameResolution, STM_None);
 	}
 	else
 	{
 		SystemViewModel->GetSequencer()->SetPlaybackStatus(EMovieScenePlayerStatus::Stopped);
-		SystemViewModel->GetSequencer()->SetLocalTime(TargetCaptureTime, STM_None);
+		SystemViewModel->GetSequencer()->SetLocalTime(TargetCaptureTime * FrameResolution, STM_None);
 	}
 
 	return FReply::Handled();
diff --git a/Engine/Plugins/Media/ImgMedia/Source/ImgMedia/ImgMedia.Build.cs b/Engine/Plugins/Media/ImgMedia/Source/ImgMedia/ImgMedia.Build.cs
--- a/Engine/Plugins/Media/ImgMedia/Source/ImgMedia/ImgMedia.Build.cs
+++ b/Engine/Plugins/Media/ImgMedia/Source/ImgMedia/ImgMedia.Build.cs
@@ -35,6 +35,7 @@ namespace UnrealBuildTool.Rules
 					"ImgMedia/Private/Loader",
 					"ImgMedia/Private/Player",
 					"ImgMedia/Private/Readers",
+					"ImgMedia/Private/Scheduler",
 				});
 
 			PublicDependencyModuleNames.AddRange(
diff --git a/Engine/Plugins/Media/ImgMedia/Source/ImgMedia/Private/ImgMediaModule.cpp b/Engine/Plugins/Media/ImgMedia/Source/ImgMedia/Private/ImgMediaModule.cpp
--- a/Engine/Plugins/Media/ImgMedia/Source/ImgMedia/Private/ImgMediaModule.cpp
+++ b/Engine/Plugins/Media/ImgMedia/Source/ImgMedia/Private/ImgMediaModule.cpp
@@ -2,15 +2,24 @@
 
 #include "ImgMediaPrivate.h"
 
+#include "IMediaClock.h"
+#include "IMediaModule.h"
+#include "Misc/QueuedThreadPool.h"
 #include "Modules/ModuleManager.h"
 
 #include "ImgMediaPlayer.h"
+#include "ImgMediaScheduler.h"
 #include "IImgMediaModule.h"
 
 
 DEFINE_LOG_CATEGORY(LogImgMedia);
 
 
+#if USE_IMGMEDIA_DEALLOC_POOL
+	FQueuedThreadPool* GImgMediaThreadPoolSlow = nullptr;
+#endif
+
+
 /**
  * Implements the AVFMedia module.
  */
@@ -28,15 +37,56 @@ public:
 
 	virtual TSharedPtr<IMediaPlayer, ESPMode::ThreadSafe> CreatePlayer(IMediaEventSink& EventSink) override
 	{
-		return MakeShared<FImgMediaPlayer, ESPMode::ThreadSafe>(EventSink);
+		return MakeShared<FImgMediaPlayer, ESPMode::ThreadSafe>(EventSink, Scheduler.ToSharedRef());
 	}
 
 public:
 
 	//~ IModuleInterface interface
 
-	virtual void StartupModule() override { }
-	virtual void ShutdownModule() override { }
+	virtual void StartupModule() override
+	{
+		// initialize scheduler
+		Scheduler = MakeShared<FImgMediaScheduler, ESPMode::ThreadSafe>();
+		Scheduler->Initialize();
+
+		IMediaModule* MediaModule = FModuleManager::LoadModulePtr<IMediaModule>("Media");
+
+		if (MediaModule != nullptr)
+		{
+			MediaModule->GetClock().AddSink(Scheduler.ToSharedRef());
+		}
+
+		// initialize worker thread pools
+		if (FPlatformProcess::SupportsMultithreading())
+		{
+#if USE_IMGMEDIA_DEALLOC_POOL
+			// initialize dealloc thread pool
+			const int32 ThreadPoolSize = FPlatformProperties::IsServerOnly() ? 1 : FPlatformMisc::NumberOfWorkerThreadsToSpawn();
+			const uint32 StackSize = 128 * 1024;
+
+			GImgMediaThreadPoolSlow = FQueuedThreadPool::Allocate();
+			verify(GImgMediaThreadPoolSlow->Create(ThreadPoolSize, StackSize, TPri_Normal));
+#endif
+		}
+	}
+
+	virtual void ShutdownModule() override
+	{
+		Scheduler.Reset();
+
+#if USE_IMGMEDIA_DEALLOC_POOL
+		if (GImgMediaThreadPoolSlow != nullptr)
+		{
+			GImgMediaThreadPoolSlow->Destroy();
+			GImgMediaThreadPoolSlow = nullptr;
+		}
+#endif
+	}
+
+private:
+
+	TSharedPtr<FImgMediaScheduler, ESPMode::ThreadSafe> Scheduler;
 };
 
 
diff --git a/Engine/Plugins/Media/ImgMedia/Source/ImgMedia/Private/ImgMediaPrivate.h b/Engine/Plugins/Media/ImgMedia/Source/ImgMedia/Private/ImgMediaPrivate.h
--- a/Engine/Plugins/Media/ImgMedia/Source/ImgMedia/Private/ImgMediaPrivate.h
+++ b/Engine/Plugins/Media/ImgMedia/Source/ImgMedia/Private/ImgMediaPrivate.h
@@ -4,8 +4,11 @@
 
 #include "ImgMediaSettings.h"
 #include "Logging/LogMacros.h"
+#include "Stats/Stats.h"
 #include "UObject/NameTypes.h"
 
+class FQueuedThreadPool;
+
 
 /** The OpenEXR image reader is supported on macOS and Windows only. */
 #define IMGMEDIA_EXR_SUPPORTED_PLATFORM (PLATFORM_MAC || PLATFORM_WINDOWS)
@@ -13,10 +16,20 @@
 /** Default number of frames per second for image sequences. */
 #define IMGMEDIA_DEFAULT_FPS 24.0
 
+/** Whether to use a separate thread pool for image frame deallocations. */
+#define USE_IMGMEDIA_DEALLOC_POOL UE_BUILD_DEBUG
+
+
 /** Log category for the this module. */
 DECLARE_LOG_CATEGORY_EXTERN(LogImgMedia, Log, All);
 
 
+#if USE_IMGMEDIA_DEALLOC_POOL
+	/** Thread pool used for deleting image frame buffers. */
+	extern FQueuedThreadPool* GImgMediaThreadPoolSlow;
+#endif
+
+
 namespace ImgMedia
 {
 	/** Name of the FramesPerSecondAttribute media option. */
diff --git a/Engine/Plugins/Media/ImgMedia/Source/ImgMedia/Private/Loader/ImgMediaLoader.cpp b/Engine/Plugins/Media/ImgMedia/Source/ImgMedia/Private/Loader/ImgMediaLoader.cpp
--- a/Engine/Plugins/Media/ImgMedia/Source/ImgMedia/Private/Loader/ImgMediaLoader.cpp
+++ b/Engine/Plugins/Media/ImgMedia/Source/ImgMedia/Private/Loader/ImgMediaLoader.cpp
@@ -3,6 +3,7 @@
 #include "ImgMediaLoader.h"
 #include "ImgMediaPrivate.h"
 
+#include "Algo/Reverse.h"
 #include "HAL/FileManager.h"
 #include "HAL/PlatformProcess.h"
 #include "Misc/Paths.h"
@@ -15,6 +16,7 @@
 #include "GenericImgMediaReader.h"
 #include "IImgMediaReader.h"
 #include "ImgMediaLoaderWork.h"
+#include "ImgMediaScheduler.h"
 #include "ImgMediaTextureSample.h"
 
 #if IMGMEDIA_EXR_SUPPORTED_PLATFORM
@@ -22,15 +24,23 @@
 #endif
 
 
+/** Time spent loading a new image sequence. */
+DECLARE_CYCLE_STAT(TEXT("ImgMedia Loader Load Sequence"), STAT_ImgMedia_LoaderLoadSequence, STATGROUP_Media);
+
+/** Time spent releasing cache in image loader destructor. */
+DECLARE_CYCLE_STAT(TEXT("ImgMedia Loader Release Cache"), STAT_ImgMedia_LoaderReleaseCache, STATGROUP_Media);
+
+
 /* FImgMediaLoader structors
  *****************************************************************************/
 
-FImgMediaLoader::FImgMediaLoader()
+FImgMediaLoader::FImgMediaLoader(const TSharedRef<FImgMediaScheduler, ESPMode::ThreadSafe>& InScheduler)
 	: Frames(1)
 	, ImageWrapperModule(FModuleManager::LoadModuleChecked<IImageWrapperModule>("ImageWrapper"))
 	, Initialized(false)
 	, NumLoadAhead(0)
 	, NumLoadBehind(0)
+	, Scheduler(InScheduler)
 	, SequenceDim(FIntPoint::ZeroValue)
 	, SequenceDuration(FTimespan::Zero())
 	, SequenceFps(0.0f)
@@ -40,36 +50,21 @@ FImgMediaLoader::FImgMediaLoader()
 
 FImgMediaLoader::~FImgMediaLoader()
 {
-	FScopeLock Lock(&CriticalSection);
-
-	// clear pending frames
-	PendingFrameNumbers.Empty();
-
-	// abandon queued work items
-	for (auto WorkPair : QueuedWorks)
-	{
-		auto QueuedWork = WorkPair.Value;
-		check(!WorkPool.Contains(QueuedWork));
-
-		if (GThreadPool->RetractQueuedWork(QueuedWork))
-		{
-			delete QueuedWork;
-		}
-		else
-		{
-			QueuedWork->DeleteWhenDone();
-		}
-	}
-
-	QueuedWorks.Empty();
-
-	// release work item pool
+	// clean up work item pool
 	for (auto Work : WorkPool)
 	{
 		delete Work;
 	}
 
 	WorkPool.Empty();
+
+	// release cache
+	{
+		SCOPE_CYCLE_COUNTER(STAT_ImgMedia_LoaderReleaseCache);
+
+		Frames.Empty();
+		PendingFrameNumbers.Empty();
+	}
 }
 
 
@@ -139,18 +134,37 @@ void FImgMediaLoader::GetPendingTimeRanges(TRangeSet<FTimespan>& OutRangeSet) co
 }
 
 
-void FImgMediaLoader::Initialize(const FString& SequencePath, const float FpsOverride)
+IQueuedWork* FImgMediaLoader::GetWork()
+{
+	FScopeLock Lock(&CriticalSection);
+
+	if (PendingFrameNumbers.Num() == 0)
+	{
+		return nullptr;
+	}
+
+	int32 FrameNumber = PendingFrameNumbers.Pop(false);
+	FImgMediaLoaderWork* Work = (WorkPool.Num() > 0) ? WorkPool.Pop() : new FImgMediaLoaderWork(AsShared(), Reader.ToSharedRef());
+	
+	Work->Initialize(FrameNumber, ImagePaths[FrameNumber]);
+	QueuedFrameNumbers.Add(FrameNumber);
+
+	return Work;
+}
+
+
+void FImgMediaLoader::Initialize(const FString& SequencePath, const float FpsOverride, bool Loop)
 {
 	check(!Initialized); // reinitialization not allowed for now
 
-	LoadSequence(SequencePath, FpsOverride);
+	LoadSequence(SequencePath, FpsOverride, Loop);
 	FPlatformMisc::MemoryBarrier();
 
 	Initialized = true;
 }
 
 
-bool FImgMediaLoader::RequestFrame(FTimespan Time, float PlayRate)
+bool FImgMediaLoader::RequestFrame(FTimespan Time, float PlayRate, bool Loop)
 {
 	const int32 FrameIndex = TimeToFrame(Time);
 
@@ -159,7 +173,7 @@ bool FImgMediaLoader::RequestFrame(FTimespan Time, float PlayRate)
 		return false;
 	}
 
-	Update(FrameIndex, PlayRate);
+	Update(FrameIndex, PlayRate, Loop);
 	LastRequestedFrame = FrameIndex;
 
 	return true;
@@ -186,8 +200,10 @@ void FImgMediaLoader::FrameNumbersToTimeRanges(const TArray<int32>& FrameNumbers
 }
 
 
-void FImgMediaLoader::LoadSequence(const FString& SequencePath, const float FpsOverride)
+void FImgMediaLoader::LoadSequence(const FString& SequencePath, const float FpsOverride, bool Loop)
 {
+	SCOPE_CYCLE_COUNTER(STAT_ImgMedia_LoaderLoadSequence);
+
 	if (SequencePath.IsEmpty())
 	{
 		return;
@@ -280,7 +296,7 @@ void FImgMediaLoader::LoadSequence(const FString& SequencePath, const float FpsO
 	NumLoadAhead = NumFramesToLoad - NumLoadBehind;
 
 	Frames.Empty(NumFramesToLoad);
-	Update(0, 0.0f);
+	Update(0, 0.0f, Loop);
 
 	// update info
 	Info = TEXT("Image Sequence\n");
@@ -303,14 +319,16 @@ uint32 FImgMediaLoader::TimeToFrame(FTimespan Time) const
 }
 
 
-void FImgMediaLoader::Update(int32 PlayHeadFrame, float PlayRate)
+void FImgMediaLoader::Update(int32 PlayHeadFrame, float PlayRate, bool Loop)
 {
-	// @todo ImgMedia: take PlayRate and DeltaTime into account when determining frames to load
+	// @todo gmp: ImgMedia: take PlayRate and DeltaTime into account when determining frames to load
 	
 	// determine frame numbers to be loaded
-	TArray<int32> FramesToQueue;
+	TArray<int32> FramesToLoad;
 	{
-		FramesToQueue.Empty(NumLoadAhead + NumLoadBehind);
+		const int32 NumImagePaths = ImagePaths.Num();
+
+		FramesToLoad.Empty(NumLoadAhead + NumLoadBehind);
 
 		int32 FrameOffset = (PlayRate >= 0.0f) ? 1 : -1;
 
@@ -320,98 +338,104 @@ void FImgMediaLoader::Update(int32 PlayHeadFrame, float PlayRate)
 		int32 LoadBehindCount = NumLoadBehind;
 		int32 LoadBehindIndex = PlayHeadFrame - FrameOffset;
 
+		// alternate between look ahead and look behind
 		while ((LoadAheadCount > 0) || (LoadBehindCount > 0))
 		{
-			if (LoadAheadCount-- > 0)
+			if (LoadAheadCount > 0)
 			{
 				if (LoadAheadIndex < 0)
 				{
-					LoadAheadIndex += ImagePaths.Num();
+					if (Loop)
+					{
+						LoadAheadIndex += NumImagePaths;
+					}
+					else
+					{
+						LoadAheadCount = 0;
+					}
 				}
-				else if (LoadAheadIndex >= ImagePaths.Num())
+				else if (LoadAheadIndex >= NumImagePaths)
 				{
-					LoadAheadIndex -= ImagePaths.Num();
+					if (Loop)
+					{
+						LoadAheadIndex -= NumImagePaths;
+					}
+					else
+					{
+						LoadAheadCount = 0;
+					}
+				}
+
+				if (LoadAheadCount > 0)
+				{
+					FramesToLoad.Add(LoadAheadIndex);
 				}
 
-				FramesToQueue.Add(LoadAheadIndex);
 				LoadAheadIndex += FrameOffset;
+				--LoadAheadCount;
 			}
 
-			if (LoadBehindCount-- > 0)
+			if (LoadBehindCount > 0)
 			{
 				if (LoadBehindIndex < 0)
 				{
-					LoadBehindIndex += ImagePaths.Num();
+					if (Loop)
+					{
+						LoadBehindIndex += NumImagePaths;
+					}
+					else
+					{
+						LoadBehindCount = 0;
+					}
+				}
+				else if (LoadBehindIndex >= NumImagePaths)
+				{
+					if (Loop)
+					{
+						LoadBehindIndex -= NumImagePaths;
+					}
+					else
+					{
+						LoadBehindCount = 0;
+					}
 				}
-				else if (LoadBehindIndex >= ImagePaths.Num())
+
+				if (LoadBehindCount > 0)
 				{
-					LoadBehindIndex -= ImagePaths.Num();
+					FramesToLoad.Add(LoadBehindIndex);
 				}
 
-				FramesToQueue.Add(LoadBehindIndex);
 				LoadBehindIndex -= FrameOffset;
+				--LoadBehindCount;
 			}
 		}
 	}
 
-	// create new work items
-	TArray<FImgMediaLoaderWork*> NewWorks;
-	{
-		FScopeLock ScopeLock(&CriticalSection);
-
-		// queue up new pending frames
-		PendingFrameNumbers.Empty();
+	FScopeLock ScopeLock(&CriticalSection);
 
-		for (int32 Idx = FramesToQueue.Num() - 1; Idx >= 0; --Idx)
-		{
-			PendingFrameNumbers.Push(FramesToQueue[Idx]);
-		}
+	// determine queued frame numbers that can be discarded
+	for (int32 Idx = QueuedFrameNumbers.Num() - 1; Idx >= 0; --Idx)
+	{
+		const int32 FrameNumber = QueuedFrameNumbers[Idx];
 
-		// try to abandon prior queued work if obsolete
-		for (int32 Index = QueuedFrameNumbers.Num() - 1; Index >= 0; --Index)
+		if (!FramesToLoad.Contains(FrameNumber))
 		{
-			int32 FrameNumber = QueuedFrameNumbers[Index];
-
-			if (!PendingFrameNumbers.Contains(FrameNumber))
-			{
-				FImgMediaLoaderWork*& ObsoleteWork = QueuedWorks.FindChecked(FrameNumber);
-
-				if (GThreadPool->RetractQueuedWork(ObsoleteWork))
-				{
-					QueuedWorks.Remove(FrameNumber);
-					WorkPool.Push(ObsoleteWork);
-				}
-
-				QueuedFrameNumbers.RemoveAtSwap(Index);
-			}
+			QueuedFrameNumbers.RemoveAtSwap(Idx);
 		}
+	}
 
-		// create or reuse work items
-		int32 NewWorkCount = GThreadPool->GetNumThreads() - QueuedWorks.Num();
+	// determine frame numbers that need to be cached
+	PendingFrameNumbers.Empty();
 
-		while ((NewWorkCount > 0) && (PendingFrameNumbers.Num() > 0))
+	for (int32 FrameNumber : FramesToLoad)
+	{
+		if ((Frames.FindAndTouch(FrameNumber) == nullptr) && !QueuedFrameNumbers.Contains(FrameNumber))
 		{
-			const int32 PendingFrameNumber = PendingFrameNumbers.Pop(false);
-
-			if ((Frames.FindAndTouch(PendingFrameNumber) == nullptr) && !QueuedWorks.Contains(PendingFrameNumber))
-			{
-				FImgMediaLoaderWork* NewWork = (WorkPool.Num() > 0) ? WorkPool.Pop() : new FImgMediaLoaderWork(*this, Reader.ToSharedRef());
-				NewWork->Initialize(PendingFrameNumber, ImagePaths[PendingFrameNumber]);
-				NewWorks.Add(NewWork);
-
-				QueuedFrameNumbers.Add(PendingFrameNumber);
-				QueuedWorks.Add(PendingFrameNumber, NewWork);
-				
-				--NewWorkCount;
-			}
+			PendingFrameNumbers.Add(FrameNumber);
 		}
 	}
 
-	// enqueue new work
-	for (auto& NewWork : NewWorks)
-	{
-		GThreadPool->AddQueuedWork(NewWork);
-	}
+	Algo::Reverse(PendingFrameNumbers);
 }
 
 
@@ -420,50 +444,16 @@ void FImgMediaLoader::Update(int32 PlayHeadFrame, float PlayRate)
 
 void FImgMediaLoader::NotifyWorkComplete(FImgMediaLoaderWork& CompletedWork, int32 FrameNumber, const TSharedPtr<FImgMediaFrame, ESPMode::ThreadSafe>& Frame)
 {
-	FImgMediaLoaderWork* NewWork = nullptr;
+	FScopeLock Lock(&CriticalSection);
 
-	// process completed work item
+	// if frame is still needed, add it to the cache
+	if (QueuedFrameNumbers.Remove(FrameNumber) > 0)
 	{
-		FScopeLock Lock(&CriticalSection);
-
-		check(QueuedWorks.FindChecked(FrameNumber) == &CompletedWork);
-		QueuedWorks.Remove(FrameNumber);
-
-		// if frame is still needed, add it to the cache
-		if (QueuedFrameNumbers.Remove(FrameNumber) > 0)
-		{
-			if (Frame.IsValid())
-			{
-				Frames.Add(FrameNumber, Frame);
-			}
-		}
-
-		// reuse work for next frame, or return it to pool
-		while (PendingFrameNumbers.Num() > 0)
-		{
-			int32 PendingFrameNumber = PendingFrameNumbers.Pop(false);
-
-			if ((Frames.FindAndTouch(PendingFrameNumber) == nullptr) && !QueuedWorks.Contains(PendingFrameNumber))
-			{
-				NewWork = &CompletedWork;
-				NewWork->Initialize(PendingFrameNumber, ImagePaths[PendingFrameNumber]);
-				
-				QueuedFrameNumbers.Add(PendingFrameNumber);
-				QueuedWorks.Add(PendingFrameNumber, NewWork);
-
-				break;
-			}
-		}
-
-		if (NewWork == nullptr)
+		if (Frame.IsValid())
 		{
-			WorkPool.Push(&CompletedWork);
+			Frames.Add(FrameNumber, Frame);
 		}
 	}
 
-	// queue new work
-	if (NewWork != nullptr)
-	{
-		GThreadPool->AddQueuedWork(NewWork);
-	}
+	WorkPool.Push(&CompletedWork);
 }
diff --git a/Engine/Plugins/Media/ImgMedia/Source/ImgMedia/Private/Loader/ImgMediaLoader.h b/Engine/Plugins/Media/ImgMedia/Source/ImgMedia/Private/Loader/ImgMediaLoader.h
--- a/Engine/Plugins/Media/ImgMedia/Source/ImgMedia/Private/Loader/ImgMediaLoader.h
+++ b/Engine/Plugins/Media/ImgMedia/Source/ImgMedia/Private/Loader/ImgMediaLoader.h
@@ -5,13 +5,14 @@
 #include "CoreMinimal.h"
 #include "Containers/LruCache.h"
 #include "Containers/Queue.h"
-
-#include "IImgMediaLoader.h"
+#include "Templates/SharedPointer.h"
 
 class FImgMediaLoaderWork;
+class FImgMediaScheduler;
 class FImgMediaTextureSample;
 class IImageWrapperModule;
 class IImgMediaReader;
+class IQueuedWork;
 
 struct FImgMediaFrame;
 
@@ -20,12 +21,16 @@ struct FImgMediaFrame;
  * Loads image sequence frames from disk.
  */
 class FImgMediaLoader
-	: protected IImgMediaLoader
+	: public TSharedFromThis<FImgMediaLoader, ESPMode::ThreadSafe>
 {
 public:
 
-	/** Default constructor. */
-	FImgMediaLoader();
+	/**
+	 * Create and initialize a new instance.
+	 *
+	 * @param InScheduler The scheduler for image loading.
+	 */
+	FImgMediaLoader(const TSharedRef<FImgMediaScheduler, ESPMode::ThreadSafe>& InScheduler);
 
 	/** Virtual destructor. */
 	virtual ~FImgMediaLoader();
@@ -137,14 +142,24 @@ public:
 		return SequenceFps;
 	}
 
+	/**
+	 * Get the next work item.
+	 *
+	 * This method is called by the scheduler.
+	 *
+	 * @return The work item, or nullptr if no work is available.
+	 */
+	IQueuedWork* GetWork();
+
 	/**
 	 * Initialize the image sequence loader.
 	 *
 	 * @param SequencePath Path to the image sequence.
 	 * @param FpsOverride The frames per second to use (0.0 = do not override).
+	 * @param Loop Whether the cache should loop around.
 	 * @see IsInitialized
 	 */
-	void Initialize(const FString& SequencePath, const float FpsOverride);
+	void Initialize(const FString& SequencePath, const float FpsOverride, bool Loop);
 
 	/**
 	 * Whether this loader has been initialized yet.
@@ -157,14 +172,24 @@ public:
 		return Initialized;
 	}
 
+	/**
+	 * Notify the loader that a work item completed.
+	 *
+	 * @param CompletedWork The work item that completed.
+	 * @param FrameNumber Number of the frame that was read.
+	 * @param Frame The frame that was read, or nullptr if reading failed.
+	 */
+	void NotifyWorkComplete(FImgMediaLoaderWork& CompletedWork, int32 FrameNumber, const TSharedPtr<FImgMediaFrame, ESPMode::ThreadSafe>& Frame);
+
 	/**
 	 * Asynchronously request the image frame at the specified time.
 	 *
 	 * @param Time The time of the image frame to request (relative to the beginning of the sequence).
 	 * @param PlayRate The current play rate (used by the look-ahead logic).
+	 * @param Loop Whether the cache should loop around.
 	 * @return bool if the frame will be loaded, false otherwise.
 	 */
-	bool RequestFrame(FTimespan Time, float PlayRate);
+	bool RequestFrame(FTimespan Time, float PlayRate, bool Loop);
 
 protected:
 
@@ -181,8 +206,9 @@ protected:
 	 *
 	 * @param SequencePath Path to the image sequence.
 	 * @param FpsOverride The frames per second to use (0.0 = do not override).
+	 * @param Loop Whether the cache should loop around.
 	 */
-	void LoadSequence(const FString& SequencePath, const float FpsOverride);
+	void LoadSequence(const FString& SequencePath, const float FpsOverride, bool Loop);
 
 	/**
 	 * Get the frame number corresponding to the specified play head time.
@@ -197,14 +223,9 @@ protected:
 	 *
 	 * @param PlayHeadFrame Current play head frame number.
 	 * @param PlayRate Current play rate.
+	 * @param Loop Whether the cache should loop around.
 	 */
-	void Update(int32 PlayHeadFrame, float PlayRate);
-
-protected:
-
-	//~ IImgMediaLoader interface
-
-	virtual void NotifyWorkComplete(FImgMediaLoaderWork& CompletedWork, int32 FrameNumber, const TSharedPtr<FImgMediaFrame, ESPMode::ThreadSafe>& Frame) override;
+	void Update(int32 PlayHeadFrame, float PlayRate, bool Loop);
 
 private:
 
@@ -235,6 +256,9 @@ private:
 	/** The image sequence reader to use. */
 	TSharedPtr<IImgMediaReader, ESPMode::ThreadSafe> Reader;
 
+	/** The scheduler for image loading. */
+	TSharedPtr<FImgMediaScheduler, ESPMode::ThreadSafe> Scheduler;
+
 	/** Width and height of the image sequence (in pixels) .*/
 	FIntPoint SequenceDim;
 
@@ -246,8 +270,6 @@ private:
 
 private:
 
-	TArray<FImgMediaLoaderWork*> AbandonedWorks;
-
 	/** Index of the previously requested frame. */
 	int32 LastRequestedFrame;
 
@@ -257,9 +279,6 @@ private:
 	/** Collection of frame numbers that are being read. */
 	TArray<int32> QueuedFrameNumbers;
 
-	/** Maps queued frame numbers to work items. */
-	TMap<int32, FImgMediaLoaderWork*> QueuedWorks;
-
 	/** Object pool for reusable work items. */
 	TArray<FImgMediaLoaderWork*> WorkPool;
 };
diff --git a/Engine/Plugins/Media/ImgMedia/Source/ImgMedia/Private/Loader/ImgMediaLoaderWork.cpp b/Engine/Plugins/Media/ImgMedia/Source/ImgMedia/Private/Loader/ImgMediaLoaderWork.cpp
--- a/Engine/Plugins/Media/ImgMedia/Source/ImgMedia/Private/Loader/ImgMediaLoaderWork.cpp
+++ b/Engine/Plugins/Media/ImgMedia/Source/ImgMedia/Private/Loader/ImgMediaLoaderWork.cpp
@@ -5,18 +5,26 @@
 
 #include "Misc/ScopeLock.h"
 
-#include "IImgMediaLoader.h"
 #include "IImgMediaReader.h"
+#include "ImgMediaLoader.h"
+
+
+/** Time spent abandoning worker threads. */
+DECLARE_CYCLE_STAT(TEXT("ImgMedia Loader Abadon Work"), STAT_ImgMedia_LoaderAbandonWork, STATGROUP_Media);
+
+/** Time spent finalizing worker threads. */
+DECLARE_CYCLE_STAT(TEXT("ImgMedia Loader Finalize Work"), STAT_ImgMedia_LoaderFinalizeWork, STATGROUP_Media);
+
+/** Time spent reading image frames in worker threads. */
+DECLARE_CYCLE_STAT(TEXT("ImgMedia Loader Read Frame"), STAT_ImgMedia_LoaderReadFrame, STATGROUP_Media);
 
 
 /* FImgMediaLoaderWork structors
  *****************************************************************************/
 
-FImgMediaLoaderWork::FImgMediaLoaderWork(IImgMediaLoader& InOwner, const TSharedRef<IImgMediaReader, ESPMode::ThreadSafe>& InReader)
-	: AutoDelete(false)
-	, Done(false)
-	, FrameNumber(INDEX_NONE)
-	, Owner(InOwner)
+FImgMediaLoaderWork::FImgMediaLoaderWork(const TSharedRef<FImgMediaLoader, ESPMode::ThreadSafe>& InOwner, const TSharedRef<IImgMediaReader, ESPMode::ThreadSafe>& InReader)
+	: FrameNumber(INDEX_NONE)
+	, OwnerPtr(InOwner)
 	, Reader(InReader)
 { }
 
@@ -26,26 +34,8 @@ FImgMediaLoaderWork::FImgMediaLoaderWork(IImgMediaLoader& InOwner, const TShared
 
 void FImgMediaLoaderWork::Initialize(int32 InFrameNumber, const FString& InImagePath)
 {
-	check(!AutoDelete);
-
 	FrameNumber = InFrameNumber;
 	ImagePath = InImagePath;
-	Done = false;
-}
-
-
-void FImgMediaLoaderWork::DeleteWhenDone()
-{
-	FScopeLock Lock(&CriticalSection);
-
-	if (Done)
-	{
-		delete this;
-	}
-	else
-	{
-		AutoDelete = true;
-	}
 }
 
 
@@ -54,7 +44,9 @@ void FImgMediaLoaderWork::DeleteWhenDone()
 
 void FImgMediaLoaderWork::Abandon()
 {
-	// not supported
+	SCOPE_CYCLE_COUNTER(STAT_ImgMedia_LoaderAbandonWork);
+
+	Finalize(nullptr);
 }
 
 
@@ -68,6 +60,8 @@ void FImgMediaLoaderWork::DoThreadedWork()
 	}
 	else
 	{
+		SCOPE_CYCLE_COUNTER(STAT_ImgMedia_LoaderReadFrame);
+
 		// read the image frame
 		Frame = new FImgMediaFrame();
 
@@ -78,16 +72,26 @@ void FImgMediaLoaderWork::DoThreadedWork()
 		}
 	}
 
-	// notify owner, or shut down
-	FScopeLock Lock(&CriticalSection);
+	SCOPE_CYCLE_COUNTER(STAT_ImgMedia_LoaderFinalizeWork);
+
+	Finalize(Frame);
+}
+
+
+/* FImgMediaLoaderWork implementation
+ *****************************************************************************/
+
+void FImgMediaLoaderWork::Finalize(FImgMediaFrame* Frame)
+{
+	TSharedPtr<FImgMediaLoader, ESPMode::ThreadSafe> Owner = OwnerPtr.Pin();
 
-	if (AutoDelete)
+	if (Owner.IsValid())
 	{
-		delete this;
+		Owner->NotifyWorkComplete(*this, FrameNumber, MakeShareable(Frame));
 	}
 	else
 	{
-		Done = true;
-		Owner.NotifyWorkComplete(*this, FrameNumber, MakeShareable(Frame));
+		delete Frame;
+		delete this; // owner is gone, self destruct!
 	}
 }
diff --git a/Engine/Plugins/Media/ImgMedia/Source/ImgMedia/Private/Loader/ImgMediaLoaderWork.h b/Engine/Plugins/Media/ImgMedia/Source/ImgMedia/Private/Loader/ImgMediaLoaderWork.h
--- a/Engine/Plugins/Media/ImgMedia/Source/ImgMedia/Private/Loader/ImgMediaLoaderWork.h
+++ b/Engine/Plugins/Media/ImgMedia/Source/ImgMedia/Private/Loader/ImgMediaLoaderWork.h
@@ -6,9 +6,10 @@
 #include "Misc/IQueuedWork.h"
 
 class FImgMediaLoader;
-class IImgMediaLoader;
 class IImgMediaReader;
 
+struct FImgMediaFrame;
+
 
 /**
  * Loads a single image frame from disk.
@@ -24,21 +25,10 @@ public:
 	 * @param InOwner The loader that created this work item.
 	 * @param InReader The image reader to use.
 	 */
-	FImgMediaLoaderWork(IImgMediaLoader& InOwner, const TSharedRef<IImgMediaReader, ESPMode::ThreadSafe>& InReader);
+	FImgMediaLoaderWork(const TSharedRef<FImgMediaLoader, ESPMode::ThreadSafe>& InOwner, const TSharedRef<IImgMediaReader, ESPMode::ThreadSafe>& InReader);
 
 public:
 
-	/**
-	 * Delete this work item when it is done.
-	 *
-	 * This method will mark the work item for shutdown. When the work is
-	 * completed, the work item will automatically delete itself. No other
-	 * methods may be called on this object after calling this one.
-	 *
-	 * @see Initialize
-	 */
-	void DeleteWhenDone();
-
 	/**
 	 * Initialize this work item.
 	 *
@@ -55,16 +45,16 @@ public:
 	virtual void Abandon() override;
 	virtual void DoThreadedWork() override;
 
-private:
-
-	/** Whether the work item should delete itself after it is complete. */
-	bool AutoDelete;
+protected:
 
-	/** Whether the work was done. */
-	bool Done;
+	/**
+	 * Notify the owner of this work item, or self destruct.
+	 *
+	 * @param Frame The frame that was loaded by this work item.
+	 */
+	void Finalize(FImgMediaFrame* Frame);
 
-	/** Critical section for synchronizing access to Owner. */
-	FCriticalSection CriticalSection;
+private:
 
 	/** The number of the image frame. */
 	int32 FrameNumber;
@@ -73,7 +63,7 @@ private:
 	FString ImagePath;
 
 	/** The loader that created this reader task. */
-	IImgMediaLoader& Owner;
+	TWeakPtr<FImgMediaLoader, ESPMode::ThreadSafe> OwnerPtr;
 
 	/** The image sequence reader to use. */
 	TSharedPtr<IImgMediaReader, ESPMode::ThreadSafe> Reader;
diff --git a/Engine/Plugins/Media/ImgMedia/Source/ImgMedia/Private/Player/ImgMediaPlayer.cpp b/Engine/Plugins/Media/ImgMedia/Source/ImgMedia/Private/Player/ImgMediaPlayer.cpp
--- a/Engine/Plugins/Media/ImgMedia/Source/ImgMedia/Private/Player/ImgMediaPlayer.cpp
+++ b/Engine/Plugins/Media/ImgMedia/Source/ImgMedia/Private/Player/ImgMediaPlayer.cpp
@@ -12,16 +12,24 @@
 #include "UObject/WeakObjectPtr.h"
 
 #include "ImgMediaLoader.h"
+#include "ImgMediaScheduler.h"
 #include "ImgMediaSettings.h"
 #include "ImgMediaTextureSample.h"
 
 #define LOCTEXT_NAMESPACE "FImgMediaPlayer"
 
 
-/* FExrVideoPlayer structors
+/** Time spent closing image media players. */
+DECLARE_CYCLE_STAT(TEXT("ImgMedia Player Close"), STAT_ImgMedia_PlayerClose, STATGROUP_Media);
+
+/** Time spent in image media player input tick. */
+DECLARE_CYCLE_STAT(TEXT("ImgMedia Player TickInput"), STAT_ImgMedia_PlayerTickInput, STATGROUP_Media);
+
+
+/* FImgMediaPlayer structors
  *****************************************************************************/
 
-FImgMediaPlayer::FImgMediaPlayer(IMediaEventSink& InEventSink)
+FImgMediaPlayer::FImgMediaPlayer(IMediaEventSink& InEventSink, const TSharedRef<FImgMediaScheduler, ESPMode::ThreadSafe>& InScheduler)
 	: CurrentDuration(FTimespan::Zero())
 	, CurrentRate(0.0f)
 	, CurrentState(EMediaState::Closed)
@@ -29,7 +37,9 @@ FImgMediaPlayer::FImgMediaPlayer(IMediaEventSink& InEventSink)
 	, EventSink(InEventSink)
 	, LastFetchTime(FTimespan::MinValue())
 	, PlaybackRestarted(false)
+	, Scheduler(InScheduler)
 	, SelectedVideoTrack(INDEX_NONE)
+	, ShouldLoop(false)
 { }
 
 
@@ -44,11 +54,14 @@ FImgMediaPlayer::~FImgMediaPlayer()
 
 void FImgMediaPlayer::Close()
 {
+	SCOPE_CYCLE_COUNTER(STAT_ImgMedia_PlayerClose);
+
 	if (!Loader.IsValid())
 	{
 		return;
 	}
 
+	Scheduler->UnregisterLoader(Loader.ToSharedRef());
 	Loader.Reset();
 
 	CurrentDuration = FTimespan::Zero();
@@ -152,12 +165,13 @@ bool FImgMediaPlayer::Open(const FString& Url, const IMediaOptions* Options)
 	}
 
 	// initialize image loader on a separate thread
-	Loader = MakeShared<FImgMediaLoader, ESPMode::ThreadSafe>();
+	Loader = MakeShared<FImgMediaLoader, ESPMode::ThreadSafe>(Scheduler.ToSharedRef());
+	Scheduler->RegisterLoader(Loader.ToSharedRef());
 
 	const float FpsOverride = (Options != nullptr) ? Options->GetMediaOption(ImgMedia::FramesPerSecondOverrideOption, 0.0f) : 0.0f;
 	const FString SequencePath = Url.RightChop(6);
 
-	Async<void>(EAsyncExecution::ThreadPool, [FpsOverride, LoaderPtr = TWeakPtr<FImgMediaLoader, ESPMode::ThreadSafe>(Loader), Proxy, SequencePath]()
+	Async<void>(EAsyncExecution::ThreadPool, [FpsOverride, LoaderPtr = TWeakPtr<FImgMediaLoader, ESPMode::ThreadSafe>(Loader), Proxy, SequencePath, Loop = ShouldLoop]()
 	{
 		TSharedPtr<FImgMediaLoader, ESPMode::ThreadSafe> PinnedLoader = LoaderPtr.Pin();
 
@@ -170,7 +184,7 @@ bool FImgMediaPlayer::Open(const FString& Url, const IMediaOptions* Options)
 				ProxyPath = SequencePath; // fall back to root folder
 			}
 
-			PinnedLoader->Initialize(ProxyPath, FpsOverride);
+			PinnedLoader->Initialize(ProxyPath, FpsOverride, Loop);
 		}
 	});
 
@@ -186,6 +200,8 @@ bool FImgMediaPlayer::Open(const TSharedRef<FArchive, ESPMode::ThreadSafe>& /*Ar
 
 void FImgMediaPlayer::TickInput(FTimespan DeltaTime, FTimespan /*Timecode*/)
 {
+	SCOPE_CYCLE_COUNTER(STAT_ImgMedia_PlayerTickInput);
+
 	if (!Loader.IsValid() || (CurrentState == EMediaState::Error))
 	{
 		return;
@@ -252,7 +268,7 @@ void FImgMediaPlayer::TickInput(FTimespan DeltaTime, FTimespan /*Timecode*/)
 	// update image loader
 	if (SelectedVideoTrack != INDEX_NONE)
 	{
-		Loader->RequestFrame(CurrentTime, CurrentRate);
+		Loader->RequestFrame(CurrentTime, CurrentRate, ShouldLoop);
 	}
 }
 
@@ -404,7 +420,7 @@ bool FImgMediaPlayer::Seek(const FTimespan& Time)
 
 	if (CurrentState == EMediaState::Paused)
 	{
-		Loader->RequestFrame(CurrentTime, CurrentRate);
+		Loader->RequestFrame(CurrentTime, CurrentRate, ShouldLoop);
 	}
 
 	EventSink.ReceiveMediaEvent(EMediaEvent::SeekCompleted);
diff --git a/Engine/Plugins/Media/ImgMedia/Source/ImgMedia/Private/Player/ImgMediaPlayer.h b/Engine/Plugins/Media/ImgMedia/Source/ImgMedia/Private/Player/ImgMediaPlayer.h
--- a/Engine/Plugins/Media/ImgMedia/Source/ImgMedia/Private/Player/ImgMediaPlayer.h
+++ b/Engine/Plugins/Media/ImgMedia/Source/ImgMedia/Private/Player/ImgMediaPlayer.h
@@ -11,6 +11,7 @@
 #include "IMediaView.h"
 
 class FImgMediaLoader;
+class FImgMediaScheduler;
 class IImgMediaReader;
 class IMediaEventSink;
 class IMediaTextureSample;
@@ -33,8 +34,9 @@ public:
 	 * Create and initialize a new instance.
 	 *
 	 * @param InEventSink The object that receives media events from this player.
+	 * @param InScheduler The image loading scheduler to use.
 	 */
-	FImgMediaPlayer(IMediaEventSink& InEventSink);
+	FImgMediaPlayer(IMediaEventSink& InEventSink, const TSharedRef<FImgMediaScheduler, ESPMode::ThreadSafe>& InScheduler);
 
 	/** Virtual destructor. */
 	virtual ~FImgMediaPlayer();
@@ -140,6 +142,9 @@ private:
 	/** If playback just restarted from the Stopped state. */
 	bool PlaybackRestarted;
 
+	/** The scheduler for image loading. */
+	TSharedPtr<FImgMediaScheduler, ESPMode::ThreadSafe> Scheduler;
+
 	/** Index of the selected video track. */
 	int32 SelectedVideoTrack;
 
diff --git a/Engine/Plugins/Media/ImgMedia/Source/ImgMedia/Private/Readers/ExrImgMediaReader.cpp b/Engine/Plugins/Media/ImgMedia/Source/ImgMedia/Private/Readers/ExrImgMediaReader.cpp
--- a/Engine/Plugins/Media/ImgMedia/Source/ImgMedia/Private/Readers/ExrImgMediaReader.cpp
+++ b/Engine/Plugins/Media/ImgMedia/Source/ImgMedia/Private/Readers/ExrImgMediaReader.cpp
@@ -4,6 +4,7 @@
 
 #if IMGMEDIA_EXR_SUPPORTED_PLATFORM
 
+#include "Async/Async.h"
 #include "Misc/Paths.h"
 #include "OpenExrWrapper.h"
 #include "UObject/Class.h"
@@ -50,12 +51,37 @@ bool FExrImgMediaReader::ReadFrame(const FString& ImagePath, FImgMediaFrame& Out
 		return false;
 	}
 
-	void* Buffer = FMemory::Malloc(Dim.X * Dim.Y * sizeof(uint16) * 4);
-	
+	// allocate frame buffer
+	const SIZE_T BufferSize = Dim.X * Dim.Y * sizeof(uint16) * 4;
+	void* Buffer = FMemory::Malloc(BufferSize, PLATFORM_CACHE_LINE_SIZE);
+
+	auto BufferDeleter = [BufferSize](void* ObjectToDelete) {
+#if USE_IMGMEDIA_DEALLOC_POOL
+		// free buffers on the thread pool, because memory allocators may perform
+		// expensive operations, such as filling the memory with debug values
+		TFunction<void()> FreeBufferTask = [ObjectToDelete, BufferSize]()
+		{
+#endif
+			FMemory::Free(ObjectToDelete);
+#if USE_IMGMEDIA_DEALLOC_POOL
+		};
+
+		if (GImgMediaThreadPoolSlow != nullptr)
+		{
+			AsyncPool(*GImgMediaThreadPoolSlow, FreeBufferTask);
+		}
+		else
+		{
+			FreeBufferTask();
+		}
+#endif
+	};
+
+	// read frame data
 	InputFile.SetFrameBuffer(Buffer, Dim);
 	InputFile.ReadPixels(0, Dim.Y - 1);
 
-	OutFrame.Data = MakeShareable(Buffer, [](void* ObjectToDelete) { FMemory::Free(ObjectToDelete); });
+	OutFrame.Data = MakeShareable(Buffer, MoveTemp(BufferDeleter));
 	OutFrame.Format = EMediaTextureSampleFormat::FloatRGBA;
 	OutFrame.Stride = Dim.X * sizeof(unsigned short) * 4;
 
diff --git a/Engine/Plugins/Media/ImgMedia/Source/ImgMediaFactory/Private/ImgMediaSettings.cpp b/Engine/Plugins/Media/ImgMedia/Source/ImgMediaFactory/Private/ImgMediaSettings.cpp
--- a/Engine/Plugins/Media/ImgMedia/Source/ImgMediaFactory/Private/ImgMediaSettings.cpp
+++ b/Engine/Plugins/Media/ImgMedia/Source/ImgMediaFactory/Private/ImgMediaSettings.cpp
@@ -10,6 +10,8 @@ UImgMediaSettings::UImgMediaSettings()
 	: DefaultFps(24.0f)
 	, CacheBehindPercentage(25)
 	, CacheSizeGB(1.0f)
+	, CacheThreads(8)
+	, CacheThreadStackSizeKB(128)
 	, ExrDecoderThreads(0)
 	, DefaultProxy(TEXT("proxy"))
 	, UseDefaultProxy(false)
diff --git a/Engine/Plugins/Media/ImgMedia/Source/ImgMediaFactory/Public/ImgMediaSettings.h b/Engine/Plugins/Media/ImgMedia/Source/ImgMediaFactory/Public/ImgMediaSettings.h
--- a/Engine/Plugins/Media/ImgMedia/Source/ImgMediaFactory/Public/ImgMediaSettings.h
+++ b/Engine/Plugins/Media/ImgMedia/Source/ImgMediaFactory/Public/ImgMediaSettings.h
@@ -35,6 +35,14 @@ public:
 	UPROPERTY(config, EditAnywhere, Category=Caching, meta=(ClampMin=0))
 	float CacheSizeGB;
 
+	/** Maximum number of image caching threads to use (0 = number of cores, default = 8). */
+	UPROPERTY(config, EditAnywhere, Category=Caching, meta=(ClampMin=0))
+	int32 CacheThreads;
+
+	/** Size of the stack for each caching thread (in kB; default = 128). */
+	UPROPERTY(config, EditAnywhere, Category=Caching, meta=(ClampMin=128), AdvancedDisplay)
+	int32 CacheThreadStackSizeKB;
+
 public:
 
 	/** Number of worker threads to use when decoding EXR images (0 = auto). */
diff --git a/Engine/Plugins/Media/MediaCompositing/Source/MediaCompositing/MediaCompositing.Build.cs b/Engine/Plugins/Media/MediaCompositing/Source/MediaCompositing/MediaCompositing.Build.cs
--- a/Engine/Plugins/Media/MediaCompositing/Source/MediaCompositing/MediaCompositing.Build.cs
+++ b/Engine/Plugins/Media/MediaCompositing/Source/MediaCompositing/MediaCompositing.Build.cs
@@ -13,6 +13,7 @@ namespace UnrealBuildTool.Rules
 					"Engine",
 					"MovieScene",
 					"MovieSceneTracks",
+					"TimeManagement",
 				});
 
 			PrivateDependencyModuleNames.AddRange(
diff --git a/Engine/Plugins/Media/MediaCompositing/Source/MediaCompositing/Private/MovieScene/MovieSceneMediaSection.cpp b/Engine/Plugins/Media/MediaCompositing/Source/MediaCompositing/Private/MovieScene/MovieSceneMediaSection.cpp
--- a/Engine/Plugins/Media/MediaCompositing/Source/MediaCompositing/Private/MovieScene/MovieSceneMediaSection.cpp
+++ b/Engine/Plugins/Media/MediaCompositing/Source/MediaCompositing/Private/MovieScene/MovieSceneMediaSection.cpp
@@ -1,7 +1,7 @@
 // Copyright 1998-2018 Epic Games, Inc. All Rights Reserved.
 
 #include "MovieSceneMediaSection.h"
-
+#include "MovieScene.h"
 
 #define LOCTEXT_NAMESPACE "MovieSceneMediaSection"
 
@@ -17,10 +17,17 @@ UMovieSceneMediaSection::UMovieSceneMediaSection(const FObjectInitializer& Objec
 #endif
 
 	EvalOptions.CompletionMode = EMovieSceneCompletionMode::RestoreState;
+}
+
+void UMovieSceneMediaSection::PostInitProperties()
+{
+	Super::PostInitProperties();
+
+	UMovieScene* Outer = GetTypedOuter<UMovieScene>();
+	FFrameRate FrameRate = Outer ? Outer->GetFrameResolution() : FFrameRate(24, 1);
 
 	// media tracks have some preroll by default to precache frames
-	SetPreRollTime(0.5f);
+	SetPreRollFrames( (0.5 * FrameRate).RoundToFrame().Value );
 }
 
-
 #undef LOCTEXT_NAMESPACE
diff --git a/Engine/Plugins/Media/MediaCompositing/Source/MediaCompositing/Private/MovieScene/MovieSceneMediaTemplate.cpp b/Engine/Plugins/Media/MediaCompositing/Source/MediaCompositing/Private/MovieScene/MovieSceneMediaTemplate.cpp
--- a/Engine/Plugins/Media/MediaCompositing/Source/MediaCompositing/Private/MovieScene/MovieSceneMediaTemplate.cpp
+++ b/Engine/Plugins/Media/MediaCompositing/Source/MediaCompositing/Private/MovieScene/MovieSceneMediaTemplate.cpp
@@ -2,17 +2,15 @@
 
 #include "MovieSceneMediaTemplate.h"
 
-#include "MediaPlaneComponent.h"
-#include "MediaSource.h"
+#include "Math/UnrealMathUtility.h"
 #include "MediaPlayer.h"
+#include "MediaSoundComponent.h"
+#include "MediaSource.h"
 #include "MediaTexture.h"
-#include "Materials/MaterialInterface.h"
-#include "Materials/Material.h"
-#include "Materials/MaterialExpression.h"
-#include "Materials/MaterialExpressionTextureBase.h"
 #include "UObject/Package.h"
 #include "UObject/GCObject.h"
 
+#include "MovieSceneMediaData.h"
 #include "MovieSceneMediaSection.h"
 #include "MovieSceneMediaTrack.h"
 
@@ -20,380 +18,106 @@
 /* Local helpers
  *****************************************************************************/
 
-/**
- * Simple struct that tries to keep an object rooted for this object's lifetime.
- * Copying is prohibited. Moving passes ownership of the root to the LHS.
- * Cannot not guarantee rooting outside of its control.
- */
-template<typename T>
-struct TScopedRootObject
-{
-	/** Construct and roo the specified object */
-	TScopedRootObject(T* InObject)
-		: Object(InObject), bAddedToRoot(false)
-	{
-		if (InObject && !InObject->IsRooted())
-		{
-			Object->AddToRoot();
-		}
-	}
-
-	/** Copying is disabled */
-	TScopedRootObject(const TScopedRootObject&) = delete;
-	TScopedRootObject& operator=(const TScopedRootObject&) = delete;
-
-	/** Move construction/assignment passes ownership of the root to the LHS */
-	TScopedRootObject(TScopedRootObject&& In)
-		: Object(In.Object)
-		, bAddedToRoot(In.bAddedToRoot)
-	{
-		In.Object = nullptr;
-		In.bAddedToRoot = false;
-	}
-
-	TScopedRootObject& operator=(TScopedRootObject&& In)
-	{
-		Object = In.Object;
-		bAddedToRoot = In.bAddedToRoot;
-		
-		In.Object = nullptr;
-		In.bAddedToRoot = false;
-
-		return *this;
-	}
-
-	/** Destructor frees the root, if necessary */
-	~TScopedRootObject()
-	{
-		if (Object && bAddedToRoot)
-		{
-			Object->RemoveFromRoot();
-		}
-	}
-
-	T* GetObject() const
-	{
-		return Object;
-	}
-
-private:
-
-	/** The object that should be rooted */
-	T* Object;
-
-	/** Whether we added the object to the root (false if it's null, or was already rooted) */
-	bool bAddedToRoot;
-};
-
-
-// This pre-animated token is only created if we create a new mediaplayer and assign it to the property
-struct FMediaTexturePropertyPreAnimatedToken
-	: IMovieScenePreAnimatedToken
-{
-	FMediaTexturePropertyPreAnimatedToken(UMediaTexture* InPreviousPropertyValue, const TSharedPtr<FTrackInstancePropertyBindings>& InBindings)
-		: PreviousPropertyValue(InPreviousPropertyValue), PropertyBindings(InBindings)
-	{ }
-
-	/** Manual defaulted functions while we need to support compilers that can't generate/default them for us */
-	FMediaTexturePropertyPreAnimatedToken(FMediaTexturePropertyPreAnimatedToken&& In) = default;
-	FMediaTexturePropertyPreAnimatedToken& operator=(FMediaTexturePropertyPreAnimatedToken&& In) = default;
-
-	/** Non-copyable */
-	FMediaTexturePropertyPreAnimatedToken(const FMediaTexturePropertyPreAnimatedToken&) = delete;
-	FMediaTexturePropertyPreAnimatedToken& operator=(const FMediaTexturePropertyPreAnimatedToken&) = delete;
-
-	/** Virtual destructor. */
-	virtual ~FMediaTexturePropertyPreAnimatedToken() { }
-
-public:
-
-	virtual void RestoreState(UObject& RestoreObject, IMovieScenePlayer& Player)
-	{
-		UMediaTexture* PreviousTexture = PreviousPropertyValue.GetObject();
-		PropertyBindings->CallFunction<UMediaTexture*>(RestoreObject, PreviousTexture);
-	}
-
-private:
-
-	/** Keep the previous player alive so we can definitely restore it later on */
-	TScopedRootObject<UMediaTexture> PreviousPropertyValue;
-
-	/** Property bindings that allow us to set the property when we've finished evaluating */
-	TSharedPtr<FTrackInstancePropertyBindings> PropertyBindings;
-
-public:
-
-	/** Producer that can create these tokens */
-	struct FProducer
-		: IMovieScenePreAnimatedTokenProducer
-	{
-		FProducer(const TSharedPtr<FTrackInstancePropertyBindings>& InBindings)
-			: Bindings(InBindings)
-		{ }
-
-		virtual ~FProducer() { }
-
-		virtual IMovieScenePreAnimatedTokenPtr CacheExistingState(UObject& Object) const override
-		{
-			return FMediaTexturePropertyPreAnimatedToken(Bindings->GetCurrentValue<UMediaTexture*>(Object), Bindings);
-		}
-
-		TSharedPtr<FTrackInstancePropertyBindings> Bindings;
-	};
-};
-
-
-/**
- * This pre-animated token is only created if we create a new
- * media texture and assign it to an existing media player.
- */
-struct FMediaTexturePreAnimatedToken
-	: IMovieScenePreAnimatedToken
-{
-	/** Virtual destructor. */
-	virtual ~FMediaTexturePreAnimatedToken() { }
-
-	virtual void RestoreState(UObject& RestoreObject, IMovieScenePlayer& Player)
-	{
-		UMediaTexture* MediaTexture = CastChecked<UMediaTexture>(&RestoreObject);
-		// @todo gmp: fix me media framework 3.0
-		/*
-		MediaPlayer->SetVideoTexture(nullptr);
-		*/
-	}
-
-	static FMovieSceneAnimTypeID GetAnimTypeID() { return TMovieSceneAnimTypeID<FMediaTexturePreAnimatedToken>(); }
-
-	struct FProducer : IMovieScenePreAnimatedTokenProducer
-	{
-		virtual IMovieScenePreAnimatedTokenPtr CacheExistingState(UObject& Object) const override
-		{
-			return FMediaTexturePreAnimatedToken();
-		}
-	};
-};
-
-
-/**
- * Persistent data that's stored for each currently evaluating section.
- */
-struct FMediaSectionData
-	: PropertyTemplate::FSectionData, FGCObject
+struct FMediaSectionPreRollExecutionToken
+	: IMovieSceneExecutionToken
 {
-	FMediaSectionData()
-		: TemporaryMediaTexture(nullptr)
+	FMediaSectionPreRollExecutionToken(UMediaSource* InMediaSource, FTimespan InStartTimeSeconds)
+		: MediaSource(InMediaSource)
+		, StartTime(InStartTimeSeconds)
 	{ }
 
-	virtual void AddReferencedObjects(FReferenceCollector& Collector) override
-	{
-		Collector.AddReferencedObject(TemporaryMediaTexture);
-	}
-
-	UMediaTexture* GetTemporaryMediaTexture()
-	{
-		if (TemporaryMediaTexture == nullptr)
-		{
-			TemporaryMediaTexture = NewObject<UMediaTexture>(GetTransientPackage(), MakeUniqueObjectName(GetTransientPackage(), UMediaTexture::StaticClass()));
-			TemporaryMediaTexture->UpdateResource();
-
-			UMediaPlayer* MediaPlayer = NewObject<UMediaPlayer>(GetTransientPackage(), MakeUniqueObjectName(GetTransientPackage(), UMediaPlayer::StaticClass()));
-			TemporaryMediaTexture->SetMediaPlayer(MediaPlayer);
-		}
-
-		return TemporaryMediaTexture;
-	}
-
-	UMediaTexture* GetOrUpdateMediaTextureFromProperty(UObject& InObject, IMovieScenePlayer& Player)
+	virtual void Execute(const FMovieSceneContext& Context, const FMovieSceneEvaluationOperand& Operand, FPersistentEvaluationData& PersistentData, IMovieScenePlayer& Player) override
 	{
-		UMediaTexture* MediaTexture = PropertyBindings->GetCurrentValue<UMediaTexture*>(InObject);
+		using namespace PropertyTemplate;
 
-		// assign the temporary player to the property if it's null or not the temporary player
-		const bool bNeedToAssignTemporaryTexture = !MediaTexture || (TemporaryMediaTexture && MediaTexture != TemporaryMediaTexture);
+		FMovieSceneMediaData& SectionData = PersistentData.GetSectionData<FMovieSceneMediaData>();
+		UMediaPlayer* MediaPlayer = SectionData.GetMediaPlayer();
 
-		if (bNeedToAssignTemporaryTexture)
+		if (!ensure(MediaPlayer != nullptr))
 		{
-			GetTemporaryMediaTexture();
-
-			UMediaPlaneComponent* MediaPlaneComponent = CastChecked<UMediaPlaneComponent>(&InObject);
-
-			MediaPlaneComponent->SetMediaTexture(MediaTexture);
-
-			MediaTexture = TemporaryMediaTexture;
+			return;
 		}
 
-		// always assign the property last, so that anything responding to the update
-		// can have access to the media player and its video texture if needs be
-
-		if (bNeedToAssignTemporaryTexture)
+		// open the media source if necessary
+		if (MediaPlayer->GetUrl().IsEmpty())
 		{
-			// save the previous property value using a unique anim type ID to ensure we don't leave it set after evaluation
-			Player.SavePreAnimatedState(InObject, PropertyID, FMediaTexturePropertyPreAnimatedToken::FProducer(PropertyBindings));
-
-			// assign the property
-			PropertyBindings->CallFunction<UMediaTexture*>(InObject, MediaTexture);
+			SectionData.SeekOnOpen(StartTime);
+			MediaPlayer->OpenSource(MediaSource);
 		}
-
-		return MediaTexture;
 	}
 
 private:
 
-	UMediaTexture* TemporaryMediaTexture;
+	UMediaSource* MediaSource;
+	FTimespan StartTime;
 };
 
 
-struct FMediaPlayerPreAnimatedTokenProducer
-	: IMovieScenePreAnimatedTokenProducer
-{
-	virtual IMovieScenePreAnimatedTokenPtr CacheExistingState(UObject& Object) const override
-	{
-		struct FToken : IMovieScenePreAnimatedToken
-		{
-			FToken(float InRate, const FString& InOldUrl, FTimespan InSeekPosition)
-				: Rate(InRate), OldUrl(InOldUrl), SeekPosition(InSeekPosition)
-			{}
-
-			virtual void RestoreState(UObject& RestoreObject, IMovieScenePlayer& Player)
-			{
-				UMediaPlayer* LocalMediaPlayer = CastChecked<UMediaPlayer>(&RestoreObject);
-
-				if (!OldUrl.IsEmpty())
-				{
-					LocalMediaPlayer->OpenUrl(OldUrl);
-				}
-				else
-				{
-					LocalMediaPlayer->Close();
-				}
-
-				LocalMediaPlayer->SetRate(Rate);
-				LocalMediaPlayer->Seek(SeekPosition);
-			}
-
-			float Rate;
-			FString OldUrl;
-			FTimespan SeekPosition;
-		};
-
-		UMediaPlayer* MediaPlayer = CastChecked<UMediaPlayer>(&Object);
-		return FToken{ MediaPlayer->GetRate(), MediaPlayer->GetUrl(), MediaPlayer->GetTime() };
-	}
-};
-
-
-struct FMediaSectionPreRollExecutionToken
+struct FMediaSectionExecutionToken
 	: IMovieSceneExecutionToken
 {
-	FMediaSectionPreRollExecutionToken(UMediaSource* InSource, FTimespan InMediaTime)
-		: Source(InSource)
-		, MediaTime(InMediaTime)
-	{
-	}
+	FMediaSectionExecutionToken(UMediaSource* InMediaSource, FTimespan InCurrentTime)
+		: CurrentTime(InCurrentTime)
+		, MediaSource(InMediaSource)
+		, PlaybackRate(1.0f)
+	{ }
 
 	virtual void Execute(const FMovieSceneContext& Context, const FMovieSceneEvaluationOperand& Operand, FPersistentEvaluationData& PersistentData, IMovieScenePlayer& Player) override
 	{
-		using namespace PropertyTemplate;
-
-		const bool bForceTemporaryPlayer = ForceTemporaryPlayer();
+		FMovieSceneMediaData& SectionData = PersistentData.GetSectionData<FMovieSceneMediaData>();
+		UMediaPlayer* MediaPlayer = SectionData.GetMediaPlayer();
 
-		FMediaSectionData& SectionData = PersistentData.GetSectionData<FMediaSectionData>();
-		for (TWeakObjectPtr<> WeakObject : Player.FindBoundObjects(Operand))
+		if (!ensure(MediaPlayer != nullptr))
 		{
-			UObject* Object = WeakObject.Get();
-			if (!Object)
-			{
-				continue;
-			}
-
-			UMediaTexture* MediaTexture = bForceTemporaryPlayer ? SectionData.GetTemporaryMediaTexture() : SectionData.GetOrUpdateMediaTextureFromProperty(*Object, Player);
-			if (!ensure(MediaTexture))
-			{
-				continue;
-			}
-
-			// Save the previous media player state
-			//Player.SavePreAnimatedState(*MediaTexture, TMovieSceneAnimTypeID<FMediaSectionPreRollExecutionToken>(), FMediaPlayerPreAnimatedTokenProducer());
-
-			// Open the media source if necessary
-			UMediaPlayer* MediaPlayer = MediaTexture->GetMediaPlayer();
-
-			if (MediaPlayer)
-			{
-				if (MediaPlayer->GetUrl() != Source->GetUrl())
-				{
-					if (MediaPlayer->CanPlaySource(Source))
-					{
-						MediaPlayer->OpenSource(Source);
-					}
-					else
-					{
-						// @todo: log warning/error
-					}
-				}
-
-				OnUpdateFrame(MediaPlayer, Context);
-			}
+			return;
 		}
-	}
-
-protected:
 
-	virtual bool ForceTemporaryPlayer() const
-	{
-		return true;
-	}
-
-	virtual void OnUpdateFrame(UMediaPlayer* MediaPlayer, const FMovieSceneContext& Context)
-	{
-		if (MediaPlayer->GetRate() != 0.f)
+		// open the media source if necessary
+		if (MediaPlayer->GetUrl().IsEmpty())
 		{
-			MediaPlayer->SetRate(0.f);
-		}
+			SectionData.SeekOnOpen(CurrentTime);
+			MediaPlayer->OpenSource(MediaSource);
 
-		// @todo: Set this on the section itself
-		if (MediaPlayer->GetTime() != MediaTime)
-		{
-			MediaPlayer->Seek(MediaTime);
+			return;
 		}
-	}
-
-	UMediaSource* Source;
-	FTimespan MediaTime;
-};
 
+		if (MediaPlayer->IsPreparing())
+		{
+			SectionData.SeekOnOpen(CurrentTime);
 
-struct FMediaSectionExecutionToken
-	: FMediaSectionPreRollExecutionToken
-{
-	FMediaSectionExecutionToken(UMediaSource* InSource, FTimespan InMediaTime)
-		: FMediaSectionPreRollExecutionToken(InSource, InMediaTime)
-		, PlaybackRate(1.f)
-	{
-	}
+			return;
+		}
 
-	virtual bool ForceTemporaryPlayer() const
-	{
-		return false;
-	}
+		// update media player
+		const FTimespan MediaTime = CurrentTime % MediaPlayer->GetDuration();
 
-	virtual void OnUpdateFrame(UMediaPlayer* MediaPlayer, const FMovieSceneContext& Context) override
-	{
-		// Play and update the media player
 		if (Context.GetStatus() == EMovieScenePlayerStatus::Playing)
 		{
-			if (MediaPlayer->GetRate() != PlaybackRate)
+			if (!MediaPlayer->IsPlaying())
+			{
+				MediaPlayer->Seek(MediaTime);
+				MediaPlayer->SetRate(1.0f);
+			}
+			else if (Context.HasJumped())
 			{
-				MediaPlayer->SetRate(PlaybackRate);
 				MediaPlayer->Seek(MediaTime);
 			}
 		}
 		else
 		{
-			MediaPlayer->SetRate(0.0f);
+			if (MediaPlayer->IsPlaying())
+			{
+				MediaPlayer->SetRate(0.0f);
+			}
+
 			MediaPlayer->Seek(MediaTime);
 		}
 	}
 
+private:
+
+	FTimespan CurrentTime;
+	UMediaSource* MediaSource;
 	float PlaybackRate;
 };
 
@@ -402,10 +126,19 @@ struct FMediaSectionExecutionToken
  *****************************************************************************/
 
 FMovieSceneMediaSectionTemplate::FMovieSceneMediaSectionTemplate(const UMovieSceneMediaSection& InSection, const UMovieSceneMediaTrack& InTrack)
-	: PropertyData(InTrack.GetPropertyName(), InTrack.GetPropertyPath(), NAME_None, "OnMediaTextureChanged")
 {
-	Params.Source = InSection.GetMediaSource();
-	Params.SectionStartTime = InSection.GetStartTime();
+	Params.MediaSoundComponent = InSection.MediaSoundComponent;
+	Params.MediaSource = InSection.GetMediaSource();
+	Params.MediaTexture = InSection.MediaTexture;
+
+	if (InSection.GetRange().GetLowerBound().IsClosed())
+	{
+		Params.SectionStartFrame = InSection.GetRange().GetLowerBoundValue();
+	}
+	if (InSection.GetRange().GetUpperBound().IsClosed())
+	{
+		Params.SectionEndFrame = InSection.GetRange().GetUpperBoundValue();
+	}
 }
 
 
@@ -414,25 +147,31 @@ FMovieSceneMediaSectionTemplate::FMovieSceneMediaSectionTemplate(const UMovieSce
 
 void FMovieSceneMediaSectionTemplate::Evaluate(const FMovieSceneEvaluationOperand& Operand, const FMovieSceneContext& Context, const FPersistentEvaluationData& PersistentData, FMovieSceneExecutionTokens& ExecutionTokens) const
 {
-	if (!Params.Source)
+	if ((Params.MediaSource == nullptr) || Context.IsPostRoll())
 	{
 		return;
 	}
 
+	// @todo: account for start offset and video time dilation if/when these are added
+
 	if (Context.IsPreRoll())
 	{
-		// Use the end of the preroll range if we're in shot preroll, otherwise start at frame 0
-		// @todo: account for start offset and video time dilation if/when these are added
-		float FirstMediaFrameTime = Context.HasPreRollEndTime() ? Context.GetPreRollEndTime() - Params.SectionStartTime : 0.f;
-		ExecutionTokens.Add(FMediaSectionPreRollExecutionToken(Params.Source, FTimespan(int64(ETimespan::TicksPerSecond * FirstMediaFrameTime))));
+		const FFrameRate FrameRate = Context.GetFrameRate();
+		const FFrameNumber StartFrame = Context.HasPreRollEndTime() ? Context.GetPreRollEndFrame() - Params.SectionStartFrame : 0;
+		const int64 DenominatorTicks = FrameRate.Denominator * ETimespan::TicksPerSecond;
+		const int64 StartTicks = FMath::DivideAndRoundNearest(int64(StartFrame.Value * DenominatorTicks), int64(FrameRate.Numerator));
+
+		ExecutionTokens.Add(FMediaSectionPreRollExecutionToken(Params.MediaSource, FTimespan(StartTicks)));
 	}
-	else if (!Context.IsPostRoll())
+	else if (!Context.IsPostRoll() && (Context.GetTime().FrameNumber < Params.SectionEndFrame))
 	{
-		// Only update the video if we're not postrolling (which only should happen when playing in reverse through the preroll range)
+		const FFrameRate FrameRate = Context.GetFrameRate();
+		const FFrameTime CurrentTime(Context.GetTime().FrameNumber - Params.SectionStartFrame);
+		const int64 DenominatorTicks = FrameRate.Denominator * ETimespan::TicksPerSecond;
+		const int64 CurrentTicks = FMath::DivideAndRoundNearest(int64(CurrentTime.GetFrame().Value * DenominatorTicks), int64(FrameRate.Numerator));
+		const int64 CurrentSubTicks = FMath::DivideAndRoundNearest(int64(CurrentTime.GetSubFrame() * DenominatorTicks), int64(FrameRate.Numerator));
 
-		// @todo: account for start offset and video time dilation if/when these are added
-		float MediaTime = Context.GetTime() - Params.SectionStartTime;
-		ExecutionTokens.Add(FMediaSectionExecutionToken(Params.Source, FTimespan(int64(ETimespan::TicksPerSecond * MediaTime))));
+		ExecutionTokens.Add(FMediaSectionExecutionToken(Params.MediaSource, FTimespan(CurrentTicks + CurrentSubTicks)));
 	}
 }
 
@@ -445,36 +184,43 @@ UScriptStruct& FMovieSceneMediaSectionTemplate::GetScriptStructImpl() const
 
 void FMovieSceneMediaSectionTemplate::Initialize(const FMovieSceneEvaluationOperand& Operand, const FMovieSceneContext& Context, FPersistentEvaluationData& PersistentData, IMovieScenePlayer& Player) const
 {
-	// Null out any prerolling video textures right at the start of the frame, before any other video sections update through their execution tokens.
-	if (!Context.IsPreRoll())
+	FMovieSceneMediaData* SectionData = PersistentData.FindSectionData<FMovieSceneMediaData>();
+
+	if (!ensure(SectionData != nullptr))
 	{
 		return;
 	}
 
-	FMediaSectionData& SectionData = PersistentData.GetSectionData<FMediaSectionData>();
-	UMediaTexture* TemporaryMediaTexture = SectionData.GetTemporaryMediaTexture();
+	UMediaPlayer* MediaPlayer = SectionData->GetMediaPlayer();
 
-	for (TWeakObjectPtr<> WeakObject : Player.FindBoundObjects(Operand))
+	if (!ensure(MediaPlayer != nullptr))
 	{
-		if (UObject* Object = WeakObject.Get())
+		return;
+	}
+
+	const bool IsEvaluating = !(Context.IsPreRoll() || Context.IsPostRoll() || (Context.GetTime().FrameNumber >= Params.SectionEndFrame));
+
+	if (Params.MediaSoundComponent != nullptr)
+	{
+		if (IsEvaluating)
 		{
-			UMediaTexture* MediaTextureFromProperty = SectionData.PropertyBindings->GetCurrentValue<UMediaTexture*>(*Object);
-
-			// Ensure the temporary (prerolling) player has no video texture. It will get
-			// created when the video is actually evaluated for real. This will only be
-			// set if we've scrubbed backwards out of a video section into its pre-roll
-			
-			// @todo gmp: fix me media framework 3.0
-			/*
-			TemporaryMediaPlayer->SetVideoTexture(nullptr);
-			*/
-
-			if (MediaTextureFromProperty == TemporaryMediaTexture)
-			{
-				// Ensure the property is not set if we're prerolling (so we don't display preroll frames)
-				FMovieSceneAnimTypeID PropertyID = SectionData.PropertyID;
-				Player.RestorePreAnimatedState(*Object, [=](FMovieSceneAnimTypeID In){ return In == PropertyID; });
-			}
+			Params.MediaSoundComponent->SetMediaPlayer(MediaPlayer);
+		}
+		else if (Params.MediaSoundComponent->GetMediaPlayer() == MediaPlayer)
+		{
+			Params.MediaSoundComponent->SetMediaPlayer(nullptr);
+		}
+	}
+
+	if (Params.MediaTexture != nullptr)
+	{
+		if (IsEvaluating)
+		{
+			Params.MediaTexture->SetMediaPlayer(MediaPlayer);
+		}
+		else if (Params.MediaTexture->GetMediaPlayer() == MediaPlayer)
+		{
+			Params.MediaTexture->SetMediaPlayer(nullptr);
 		}
 	}
 }
@@ -482,11 +228,39 @@ void FMovieSceneMediaSectionTemplate::Initialize(const FMovieSceneEvaluationOper
 
 void FMovieSceneMediaSectionTemplate::Setup(FPersistentEvaluationData& PersistentData, IMovieScenePlayer& Player) const
 {
-	PropertyData.SetupTrack<FMediaSectionData>(PersistentData);
+	PersistentData.AddSectionData<FMovieSceneMediaData>().Setup();
 }
 
 
 void FMovieSceneMediaSectionTemplate::SetupOverrides()
 {
-	EnableOverrides(RequiresSetupFlag | RequiresInitializeFlag);
+	EnableOverrides(RequiresInitializeFlag | RequiresSetupFlag | RequiresTearDownFlag);
+}
+
+
+void FMovieSceneMediaSectionTemplate::TearDown(FPersistentEvaluationData& PersistentData, IMovieScenePlayer& Player) const
+{
+	FMovieSceneMediaData* SectionData = PersistentData.FindSectionData<FMovieSceneMediaData>();
+
+	if (!ensure(SectionData != nullptr))
+	{
+		return;
+	}
+
+	UMediaPlayer* MediaPlayer = SectionData->GetMediaPlayer();
+
+	if (!ensure(MediaPlayer != nullptr))
+	{
+		return;
+	}
+
+	if ((Params.MediaSoundComponent != nullptr) && (Params.MediaSoundComponent->GetMediaPlayer() == MediaPlayer))
+	{
+		Params.MediaSoundComponent->SetMediaPlayer(nullptr);
+	}
+
+	if ((Params.MediaTexture != nullptr) && (Params.MediaTexture->GetMediaPlayer() == MediaPlayer))
+	{
+		Params.MediaTexture->SetMediaPlayer(nullptr);
+	}
 }
diff --git a/Engine/Plugins/Media/MediaCompositing/Source/MediaCompositing/Private/MovieScene/MovieSceneMediaTemplate.h b/Engine/Plugins/Media/MediaCompositing/Source/MediaCompositing/Private/MovieScene/MovieSceneMediaTemplate.h
--- a/Engine/Plugins/Media/MediaCompositing/Source/MediaCompositing/Private/MovieScene/MovieSceneMediaTemplate.h
+++ b/Engine/Plugins/Media/MediaCompositing/Source/MediaCompositing/Private/MovieScene/MovieSceneMediaTemplate.h
@@ -9,7 +9,9 @@
 
 class UMovieSceneMediaSection;
 class UMovieSceneMediaTrack;
+class UMediaSoundComponent;
 class UMediaSource;
+class UMediaTexture;
 
 
 USTRUCT()
@@ -18,13 +20,22 @@ struct FMovieSceneMediaSectionParams
 	GENERATED_BODY()
 
 	UPROPERTY()
-	float SectionStartTime;
+	UMediaSoundComponent* MediaSoundComponent;
 
 	UPROPERTY()
-	UMediaSource* Source;
+	UMediaSource* MediaSource;
+
+	UPROPERTY()
+	UMediaTexture* MediaTexture;
 
 	UPROPERTY()
 	FString Proxy;
+
+	UPROPERTY()
+	FFrameNumber SectionStartFrame;
+
+	UPROPERTY()
+	FFrameNumber SectionEndFrame;
 };
 
 
@@ -35,7 +46,7 @@ struct FMovieSceneMediaSectionTemplate
 	GENERATED_BODY()
 
 	/** Default constructor. */
-	FMovieSceneMediaSectionTemplate() {}
+	FMovieSceneMediaSectionTemplate() { }
 
 	/**
 	 * Create and initialize a new instance.
@@ -52,14 +63,12 @@ public:
 	virtual void Evaluate(const FMovieSceneEvaluationOperand& Operand, const FMovieSceneContext& Context, const FPersistentEvaluationData& PersistentData, FMovieSceneExecutionTokens& ExecutionTokens) const override;
 	virtual UScriptStruct& GetScriptStructImpl() const override;
 	virtual void Initialize(const FMovieSceneEvaluationOperand& Operand, const FMovieSceneContext& Context, FPersistentEvaluationData& PersistentData, IMovieScenePlayer& Player) const override;
-	virtual void Setup(FPersistentEvaluationData& PersistentData, IMovieScenePlayer& Player) const override;
+	virtual void Setup(FPersistentEvaluationData& PersistentData, IMovieScenePlayer& Player) const;
 	virtual void SetupOverrides() override;
+	virtual void TearDown(FPersistentEvaluationData& PersistentData, IMovieScenePlayer& Player) const override;
 
 private:
 
-	UPROPERTY()
-	FMovieScenePropertySectionData PropertyData;
-
 	UPROPERTY()
 	FMovieSceneMediaSectionParams Params;
 };
diff --git a/Engine/Plugins/Media/MediaCompositing/Source/MediaCompositing/Private/MovieScene/MovieSceneMediaTrack.cpp b/Engine/Plugins/Media/MediaCompositing/Source/MediaCompositing/Private/MovieScene/MovieSceneMediaTrack.cpp
--- a/Engine/Plugins/Media/MediaCompositing/Source/MediaCompositing/Private/MovieScene/MovieSceneMediaTrack.cpp
+++ b/Engine/Plugins/Media/MediaCompositing/Source/MediaCompositing/Private/MovieScene/MovieSceneMediaTrack.cpp
@@ -2,8 +2,13 @@
 
 #include "MovieSceneMediaTrack.h"
 
+#include "MediaPlayer.h"
+#include "MediaSource.h"
+#include "MovieScene.h"
 #include "MovieSceneMediaSection.h"
 #include "MovieSceneMediaTemplate.h"
+#include "UObject/Package.h"
+#include "UObject/UObjectGlobals.h"
 
 
 #define LOCTEXT_NAMESPACE "MovieSceneMediaTrack"
@@ -26,12 +31,34 @@ UMovieSceneMediaTrack::UMovieSceneMediaTrack(const FObjectInitializer& ObjectIni
 }
 
 
+/* UMovieSceneMediaTrack interface
+ *****************************************************************************/
+
+UMovieSceneSection* UMovieSceneMediaTrack::AddNewMediaSourceOnRow(UMediaSource& MediaSource, FFrameNumber Time, int32 RowIndex)
+{
+	const float DefaultMediaSectionDuration = 1.0f;
+	FFrameRate FrameRate = GetTypedOuter<UMovieScene>()->GetFrameResolution();
+
+	FFrameTime DurationToUse = DefaultMediaSectionDuration * FrameRate;
+
+	// add the section
+	UMovieSceneMediaSection* NewSection = NewObject<UMovieSceneMediaSection>(this);
+
+	NewSection->InitialPlacementOnRow(MediaSections, Time, DurationToUse.FrameNumber.Value, RowIndex);
+	NewSection->SetMediaSource(&MediaSource);
+
+	MediaSections.Add(NewSection);
+
+	return NewSection;
+}
+
+
 /* UMovieScenePropertyTrack interface
  *****************************************************************************/
 
 void UMovieSceneMediaTrack::AddSection(UMovieSceneSection& Section)
 {
-	Sections.AddUnique(&Section);
+	MediaSections.Add(&Section);
 }
 
 
@@ -43,36 +70,32 @@ UMovieSceneSection* UMovieSceneMediaTrack::CreateNewSection()
 
 const TArray<UMovieSceneSection*>& UMovieSceneMediaTrack::GetAllSections() const
 {
-	return Sections;
+	return MediaSections;
 }
 
 
-void UMovieSceneMediaTrack::RemoveSection(UMovieSceneSection& Section)
+bool UMovieSceneMediaTrack::HasSection(const UMovieSceneSection& Section) const
 {
-	Sections.Remove(&Section);
+	return MediaSections.Contains(&Section);
 }
 
 
-FMovieSceneEvalTemplatePtr UMovieSceneMediaTrack::CreateTemplateForSection(const UMovieSceneSection& InSection) const
+bool UMovieSceneMediaTrack::IsEmpty() const
 {
-	return FMovieSceneMediaSectionTemplate(*CastChecked<const UMovieSceneMediaSection>(&InSection), *this);
+	return MediaSections.Num() == 0;
 }
 
 
-#if WITH_EDITORONLY_DATA
-
-FText UMovieSceneMediaTrack::GetDefaultDisplayName() const
+void UMovieSceneMediaTrack::RemoveSection(UMovieSceneSection& Section)
 {
-	return LOCTEXT("DefaultDisplayName", "Media Track");
+	MediaSections.Remove(&Section);
 }
 
 
-FName UMovieSceneMediaTrack::GetTrackName() const
+FMovieSceneEvalTemplatePtr UMovieSceneMediaTrack::CreateTemplateForSection(const UMovieSceneSection& InSection) const
 {
-	return UniqueTrackName;
+	return FMovieSceneMediaSectionTemplate(*CastChecked<const UMovieSceneMediaSection>(&InSection), *this);
 }
 
-#endif
-
 
 #undef LOCTEXT_NAMESPACE
diff --git a/Engine/Plugins/Media/MediaCompositing/Source/MediaCompositing/Public/MovieSceneMediaSection.h b/Engine/Plugins/Media/MediaCompositing/Source/MediaCompositing/Public/MovieSceneMediaSection.h
--- a/Engine/Plugins/Media/MediaCompositing/Source/MediaCompositing/Public/MovieSceneMediaSection.h
+++ b/Engine/Plugins/Media/MediaCompositing/Source/MediaCompositing/Public/MovieSceneMediaSection.h
@@ -6,7 +6,9 @@
 
 #include "MovieSceneMediaSection.generated.h"
 
+class UMediaSoundComponent;
 class UMediaSource;
+class UMediaTexture;
 
 
 /**
@@ -24,6 +26,14 @@ public:
 	UPROPERTY(EditAnywhere, Category="Media")
 	FString Proxy;
 
+	/** The media sound component that receives the track's audio output. */
+	UPROPERTY(EditAnywhere, Category="Media")
+	UMediaSoundComponent* MediaSoundComponent;
+
+	/** The media texture that receives the track's video output. */
+	UPROPERTY(EditAnywhere, Category="Media")
+	UMediaTexture* MediaTexture;
+
 public:
 
 	/**
@@ -33,6 +43,8 @@ public:
 	 */
 	UMovieSceneMediaSection(const FObjectInitializer& ObjectInitializer);
 
+	virtual void PostInitProperties() override;
+
 public:
 
 	/**
@@ -76,7 +88,7 @@ public:
 
 private:
 
-	/** The source to play with this video track */
+	/** The source to play with this video track. */
 	UPROPERTY(EditAnywhere, Category="Media")
 	UMediaSource* MediaSource;
 
diff --git a/Engine/Plugins/Media/MediaCompositing/Source/MediaCompositing/Public/MovieSceneMediaTrack.h b/Engine/Plugins/Media/MediaCompositing/Source/MediaCompositing/Public/MovieSceneMediaTrack.h
--- a/Engine/Plugins/Media/MediaCompositing/Source/MediaCompositing/Public/MovieSceneMediaTrack.h
+++ b/Engine/Plugins/Media/MediaCompositing/Source/MediaCompositing/Public/MovieSceneMediaTrack.h
@@ -3,18 +3,20 @@
 #pragma once
 
 #include "Misc/Guid.h"
-#include "Tracks/MovieScenePropertyTrack.h"
+#include "MovieSceneNameableTrack.h"
 #include "UObject/ObjectMacros.h"
 
 #include "MovieSceneMediaTrack.generated.h"
 
+class UMediaSource;
+
 
 /**
  * Implements a movie scene track for media playback.
  */
 UCLASS(MinimalAPI)
 class UMovieSceneMediaTrack
-	: public UMovieScenePropertyTrack
+	: public UMovieSceneNameableTrack
 {
 public:
 
@@ -27,6 +29,14 @@ public:
 	 */
 	UMovieSceneMediaTrack(const FObjectInitializer& ObjectInitializer);
 
+public:
+
+	/** Adds a new media source to the track. */
+	virtual UMovieSceneSection* AddNewMediaSourceOnRow(UMediaSource& MediaSource, FFrameNumber Time, int32 RowIndex);
+
+	/** Adds a new media source on the next available/non-overlapping row. */
+	virtual UMovieSceneSection* AddNewMediaSource(UMediaSource& MediaSource, FFrameNumber Time) { return AddNewMediaSourceOnRow(MediaSource, Time, INDEX_NONE); }
+
 public:
 
 	//~ UMovieScenePropertyTrack interface
@@ -35,11 +45,14 @@ public:
 	virtual UMovieSceneSection* CreateNewSection() override;
 	virtual FMovieSceneEvalTemplatePtr CreateTemplateForSection(const UMovieSceneSection& InSection) const override;
 	virtual const TArray<UMovieSceneSection*>& GetAllSections() const override;
+	virtual bool HasSection(const UMovieSceneSection& Section) const override;
+	virtual bool IsEmpty() const override;
 	virtual void RemoveSection(UMovieSceneSection& Section) override;
 	virtual bool SupportsMultipleRows() const override { return true; }
 
-#if WITH_EDITORONLY_DATA
-	virtual FText GetDefaultDisplayName() const override;
-	virtual FName GetTrackName() const override;
-#endif
+private:
+
+	/** List of all master media sections. */
+	UPROPERTY()
+	TArray<UMovieSceneSection*> MediaSections;
 };
diff --git a/Engine/Plugins/Media/MediaCompositing/Source/MediaCompositingEditor/MediaCompositingEditor.Build.cs b/Engine/Plugins/Media/MediaCompositing/Source/MediaCompositingEditor/MediaCompositingEditor.Build.cs
--- a/Engine/Plugins/Media/MediaCompositing/Source/MediaCompositingEditor/MediaCompositingEditor.Build.cs
+++ b/Engine/Plugins/Media/MediaCompositing/Source/MediaCompositingEditor/MediaCompositingEditor.Build.cs
@@ -15,6 +15,7 @@ namespace UnrealBuildTool.Rules
 					"Engine",
 					"MediaAssets",
 					"MediaCompositing",
+					"MediaUtils",
 					"MovieScene",
 					"MovieSceneTools",
 					"MovieSceneTracks",
