Copying //UE4/Dev-Sequencer to //UE4/Dev-Main (Source: //UE4/Dev-Sequencer @ 3946692)

diff --git a/Engine/Plugins/Media/MediaCompositing/Source/MediaCompositingEditor/Private/Sequencer/MediaThumbnailSection.cpp b/Engine/Plugins/Media/MediaCompositing/Source/MediaCompositingEditor/Private/Sequencer/MediaThumbnailSection.cpp
--- a/Engine/Plugins/Media/MediaCompositing/Source/MediaCompositingEditor/Private/Sequencer/MediaThumbnailSection.cpp
+++ b/Engine/Plugins/Media/MediaCompositing/Source/MediaCompositingEditor/Private/Sequencer/MediaThumbnailSection.cpp
@@ -3,41 +3,48 @@
 #include "MediaThumbnailSection.h"
 
 #include "EditorStyleSet.h"
+#include "IMediaCache.h"
+#include "IMediaTracks.h"
 #include "ISequencer.h"
 #include "MediaPlayer.h"
+#include "MediaPlayerFacade.h"
 #include "MediaSource.h"
 #include "MediaTexture.h"
+#include "MovieScene.h"
 #include "MovieSceneMediaSection.h"
+#include "MovieSceneTimeHelpers.h"
 #include "Rendering/DrawElements.h"
 #include "SequencerSectionPainter.h"
 #include "TrackEditorThumbnail/TrackEditorThumbnailPool.h"
-#include "UObject/Package.h"
-#include "UObject/UObjectGlobals.h"
 #include "Widgets/Layout/SBox.h"
 #include "Widgets/Text/STextBlock.h"
 
+#include "MovieSceneMediaData.h"
+
 
 #define LOCTEXT_NAMESPACE "FMediaThumbnailSection"
 
 
-/* FMediaSection structors
+/* FMediaThumbnailSection structors
  *****************************************************************************/
 
 FMediaThumbnailSection::FMediaThumbnailSection(UMovieSceneMediaSection& InSection, TSharedPtr<FTrackEditorThumbnailPool> InThumbnailPool, TSharedPtr<ISequencer> InSequencer)
 	: FThumbnailSection(InSequencer, InThumbnailPool, this, InSection)
+//	, MediaPlayer(nullptr)
+//	, MediaTexture(nullptr)
+	, SectionPtr(&InSection)
+	, SequencerPtr(InSequencer)
 {
-	MediaPlayer = nullptr;
-	MediaTexture = nullptr;
 	TimeSpace = ETimeSpace::Local;
 }
 
 
 FMediaThumbnailSection::~FMediaThumbnailSection()
 {
-	if (MediaPlayer != nullptr)
-	{
-		MediaPlayer->Close();
-	}
+//	if (MediaPlayer != nullptr)
+//	{
+//		MediaPlayer->Close();
+//	}
 }
 
 
@@ -46,36 +53,14 @@ FMediaThumbnailSection::~FMediaThumbnailSection()
 
 void FMediaThumbnailSection::AddReferencedObjects(FReferenceCollector& Collector)
 {
-	Collector.AddReferencedObject(MediaPlayer);
-	Collector.AddReferencedObject(MediaTexture);
+//	Collector.AddReferencedObject(MediaPlayer);
+//	Collector.AddReferencedObject(MediaTexture);
 }
 
 
 /* FThumbnailSection interface
  *****************************************************************************/
 
-FText FMediaThumbnailSection::GetSectionTitle() const
-{
-	if (UMediaSource* MediaSource = Section ? CastChecked<UMovieSceneMediaSection>(Section)->GetMediaSource() : nullptr)
-	{
-		return FText::FromString(MediaSource->GetName());
-	}
-
-	return FText::GetEmpty();
-}
-
-
-void FMediaThumbnailSection::SetSingleTime(float GlobalTime)
-{
-	UMovieSceneMediaSection* MediaSection = CastChecked<UMovieSceneMediaSection>(Section);
-	
-	if (MediaSection)
-	{
-		MediaSection->SetThumbnailReferenceOffset(GlobalTime - MediaSection->GetStartTime());
-	}
-}
-
-
 TSharedRef<SWidget> FMediaThumbnailSection::GenerateSectionWidget()
 {
 	return SNew(SBox)
@@ -90,87 +75,146 @@ TSharedRef<SWidget> FMediaThumbnailSection::GenerateSectionWidget()
 }
 
 
-float FMediaThumbnailSection::GetSectionHeight() const
+FMargin FMediaThumbnailSection::GetContentPadding() const
 {
-	return FThumbnailSection::GetSectionHeight() + 2 * 9.0f; // make space for the film border
+	return FMargin(8.0f, 15.0f);
 }
 
 
-FMargin FMediaThumbnailSection::GetContentPadding() const
+float FMediaThumbnailSection::GetSectionHeight() const
 {
-	return FMargin(8.0f, 15.0f);
+	return FThumbnailSection::GetSectionHeight() + 2 * 9.0f; // make space for the film border
 }
 
 
-void FMediaThumbnailSection::Tick(const FGeometry& AllottedGeometry, const FGeometry& ClippedGeometry, const double InCurrentTime, const float InDeltaTime)
+FText FMediaThumbnailSection::GetSectionTitle() const
 {
-	UMovieSceneMediaSection* MediaSection = Cast<UMovieSceneMediaSection>(Section);
+	UMovieSceneMediaSection* MediaSection = CastChecked<UMovieSceneMediaSection>(Section);
+	UMediaSource* MediaSource = MediaSection->GetMediaSource();
 
-	if (MediaSection != nullptr)
+	if (MediaSource == nullptr)
 	{
-		if (GetDefault<UMovieSceneUserThumbnailSettings>()->bDrawSingleThumbnails)
-		{
-			ThumbnailCache.SetSingleReferenceFrame(MediaSection->GetThumbnailReferenceOffset());
-		}
-		else
-		{
-			ThumbnailCache.SetSingleReferenceFrame(TOptional<float>());
-		}
+		return LOCTEXT("NoSequence", "Empty");
 	}
 
-	FThumbnailSection::Tick(AllottedGeometry, ClippedGeometry, InCurrentTime, InDeltaTime);
+	return FText::FromString(MediaSource->GetFName().ToString());
 }
 
 
 int32 FMediaThumbnailSection::OnPaintSection(FSequencerSectionPainter& InPainter) const
 {
-	static const FSlateBrush* FilmBorder = FEditorStyle::GetBrush("Sequencer.Section.FilmBorder");
-
+	// draw background
 	InPainter.LayerId = InPainter.PaintSectionBackground();
 
-	FVector2D LocalSectionSize = InPainter.SectionGeometry.GetLocalSize();
-
+	FVector2D SectionSize = InPainter.SectionGeometry.GetLocalSize();
 	FSlateClippingZone ClippingZone(InPainter.SectionClippingRect.InsetBy(FMargin(1.0f)));
+	
+	InPainter.DrawElements.PushClip(ClippingZone);
+	{
+		DrawFilmBorder(InPainter, SectionSize);
+	}
+	InPainter.DrawElements.PopClip();
+
+	// draw thumbnails
+	int32 LayerId = FThumbnailSection::OnPaintSection(InPainter) + 1;
+
+	UMediaPlayer* MediaPlayer = GetTemplateMediaPlayer();
+
+	if (MediaPlayer == nullptr)
+	{
+		return LayerId;
+	}
+
+	// draw overlays
+	const FTimespan MediaDuration = MediaPlayer->GetDuration();
+
+	if (MediaDuration.IsZero())
+	{
+		return LayerId;
+	}
+
+	TSharedRef<FMediaPlayerFacade, ESPMode::ThreadSafe> MediaPlayerFacade = MediaPlayer->GetPlayerFacade();
+
 	InPainter.DrawElements.PushClip(ClippingZone);
+	{
+		TRangeSet<FTimespan> CacheRangeSet;
 
-	FSlateDrawElement::MakeBox(
-		InPainter.DrawElements,
-		InPainter.LayerId++,
-		InPainter.SectionGeometry.ToPaintGeometry(FVector2D(LocalSectionSize.X - 2.0f, 7.0f), FSlateLayoutTransform(FVector2D(1.0f, 4.0f))),
-		FilmBorder,
-		InPainter.bParentEnabled ? ESlateDrawEffect::None : ESlateDrawEffect::DisabledEffect
-	);
+		MediaPlayerFacade->QueryCacheState(EMediaTrackType::Video, EMediaCacheState::Pending, CacheRangeSet);
+		DrawSampleStates(InPainter, MediaDuration, SectionSize, CacheRangeSet, FLinearColor::Gray);
 
-	FSlateDrawElement::MakeBox(
-		InPainter.DrawElements,
-		InPainter.LayerId++,
-		InPainter.SectionGeometry.ToPaintGeometry(FVector2D(LocalSectionSize.X - 2.0f, 7.0f), FSlateLayoutTransform(FVector2D(1.0f, LocalSectionSize.Y - 11.0f))),
-		FilmBorder,
-		InPainter.bParentEnabled ? ESlateDrawEffect::None : ESlateDrawEffect::DisabledEffect
-	);
+		CacheRangeSet.Empty();
 
-	InPainter.DrawElements.PopClip();
+		MediaPlayerFacade->QueryCacheState(EMediaTrackType::Video, EMediaCacheState::Loading, CacheRangeSet);
+		DrawSampleStates(InPainter, MediaDuration, SectionSize, CacheRangeSet, FLinearColor::Yellow);
 
-	int32 LayerId = FThumbnailSection::OnPaintSection(InPainter);
+		CacheRangeSet.Empty();
 
-	// @todo gmp: draw media cache state
+		MediaPlayerFacade->QueryCacheState(EMediaTrackType::Video, EMediaCacheState::Loaded, CacheRangeSet);
+		DrawSampleStates(InPainter, MediaDuration, SectionSize, CacheRangeSet, FLinearColor(0.10616, 0.48777, 0.10616));
 
-	++LayerId;
+		CacheRangeSet.Empty();
+
+		MediaPlayerFacade->QueryCacheState(EMediaTrackType::Video, EMediaCacheState::Cached, CacheRangeSet);
+		DrawSampleStates(InPainter, MediaDuration, SectionSize, CacheRangeSet, FLinearColor(0.07059, 0.32941, 0.07059));
+
+		DrawLoopIndicators(InPainter, MediaDuration, SectionSize);
+	}
+	InPainter.DrawElements.PopClip();
 
 	return LayerId;
 }
 
 
+void FMediaThumbnailSection::SetSingleTime(double GlobalTime)
+{
+	UMovieSceneMediaSection* MediaSection = CastChecked<UMovieSceneMediaSection>(Section);
+
+	if (MediaSection != nullptr)
+	{
+		double StartTime = MediaSection->GetInclusiveStartFrame() / MediaSection->GetTypedOuter<UMovieScene>()->GetFrameResolution();
+		MediaSection->SetThumbnailReferenceOffset(GlobalTime - StartTime);
+	}
+}
+
+
+void FMediaThumbnailSection::Tick(const FGeometry& AllottedGeometry, const FGeometry& ClippedGeometry, const double InCurrentTime, const float InDeltaTime)
+{
+	UMovieSceneMediaSection* MediaSection = Cast<UMovieSceneMediaSection>(Section);
+
+	if (MediaSection != nullptr)
+	{
+		if (GetDefault<UMovieSceneUserThumbnailSettings>()->bDrawSingleThumbnails)
+		{
+			ThumbnailCache.SetSingleReferenceFrame(MediaSection->GetThumbnailReferenceOffset());
+		}
+		else
+		{
+			ThumbnailCache.SetSingleReferenceFrame(TOptional<double>());
+		}
+	}
+
+	FThumbnailSection::Tick(AllottedGeometry, ClippedGeometry, InCurrentTime, InDeltaTime);
+}
+
+
 /* ICustomThumbnailClient interface
  *****************************************************************************/
 
 void FMediaThumbnailSection::Draw(FTrackEditorThumbnail& TrackEditorThumbnail)
-{
-	if ((MediaTexture == nullptr) || (MediaTexture->Resource == nullptr) || !MediaTexture->Resource->TextureRHI.IsValid())
+{/*
+	check(MediaPlayer != nullptr);
+
+	if (MediaPlayer->IsPreparing())
 	{
+		RedrawThumbnails();
+
 		return;
 	}
 
+	check(MediaTexture != nullptr);
+	check(MediaTexture->Resource != nullptr);
+	check(MediaTexture->Resource->TextureRHI.IsValid());
+
 	// get target texture resource
 	FTexture2DRHIRef Texture2DRHI = MediaTexture->Resource->TextureRHI->GetTexture2D();
 
@@ -180,9 +224,11 @@ void FMediaThumbnailSection::Draw(FTrackEditorThumbnail& TrackEditorThumbnail)
 	}
 
 	// seek media player to thumbnail position
-	float VideoTime = FMath::Max(0.0f, TrackEditorThumbnail.GetEvalPosition());
+	const float EvalPosition = FMath::Max(0.0f, TrackEditorThumbnail.GetEvalPosition());
+	const FTimespan EvalTime = int64(EvalPosition * ETimespan::TicksPerSecond);
+	const FTimespan MediaTime = EvalTime % MediaPlayer->GetDuration();
 
-	if (!MediaPlayer->Seek(FTimespan(int64(ETimespan::TicksPerSecond * VideoTime))))
+	if (!MediaPlayer->Seek(FTimespan(MediaTime)))
 	{
 		return;
 	}
@@ -195,16 +241,16 @@ void FMediaThumbnailSection::Draw(FTrackEditorThumbnail& TrackEditorThumbnail)
 	if (Sequencer.IsValid())
 	{
 		TrackEditorThumbnail.SetupFade(Sequencer->GetSequencerWidget());
-	}
+	}*/
 }
 
 
 void FMediaThumbnailSection::Setup()
-{
+{/*
 	UMovieSceneMediaSection* MediaSection = CastChecked<UMovieSceneMediaSection>(Section);
 	UMediaSource* MediaSource = MediaSection->GetMediaSource();
 
-	if (!MediaSource)
+	if (MediaSource == nullptr)
 	{
 		return;
 	}
@@ -226,13 +272,157 @@ void FMediaThumbnailSection::Setup()
 	// open latest media source
 	if (MediaPlayer->GetUrl() != MediaSource->GetUrl())
 	{
-		if (MediaPlayer->CanPlaySource(MediaSource))
+		MediaPlayer->OpenSource(MediaSource);
+	}
+
+	MediaPlayer->Pause();*/
+}
+
+
+/* FMediaThumbnailSection implementation
+ *****************************************************************************/
+
+void FMediaThumbnailSection::DrawFilmBorder(FSequencerSectionPainter& InPainter, FVector2D SectionSize) const
+{
+	static const FSlateBrush* FilmBorder = FEditorStyle::GetBrush("Sequencer.Section.FilmBorder");
+
+	// draw top film border
+	FSlateDrawElement::MakeBox(
+		InPainter.DrawElements,
+		InPainter.LayerId++,
+		InPainter.SectionGeometry.ToPaintGeometry(FVector2D(SectionSize.X - 2.0f, 7.0f), FSlateLayoutTransform(FVector2D(1.0f, 4.0f))),
+		FilmBorder,
+		InPainter.bParentEnabled ? ESlateDrawEffect::None : ESlateDrawEffect::DisabledEffect
+	);
+
+	// draw bottom film border
+	FSlateDrawElement::MakeBox(
+		InPainter.DrawElements,
+		InPainter.LayerId++,
+		InPainter.SectionGeometry.ToPaintGeometry(FVector2D(SectionSize.X - 2.0f, 7.0f), FSlateLayoutTransform(FVector2D(1.0f, SectionSize.Y - 11.0f))),
+		FilmBorder,
+		InPainter.bParentEnabled ? ESlateDrawEffect::None : ESlateDrawEffect::DisabledEffect
+	);
+}
+
+
+void FMediaThumbnailSection::DrawLoopIndicators(FSequencerSectionPainter& InPainter, FTimespan MediaDuration, FVector2D SectionSize) const
+{
+	static const FSlateBrush* GenericBrush = FCoreStyle::Get().GetBrush("GenericWhiteBox");
+
+	FFrameRate FrameResolution = Section->GetTypedOuter<UMovieScene>()->GetFrameResolution();
+	double SectionDuration = FFrameTime(MovieScene::DiscreteSize(Section->GetRange())) / FrameResolution;
+	const float MediaSizeX = MediaDuration.GetTotalSeconds() * SectionSize.X / SectionDuration;
+	float DrawOffset = 0.0f;
+
+	while (DrawOffset < SectionSize.X)
+	{
+		FSlateDrawElement::MakeBox(
+			InPainter.DrawElements,
+			InPainter.LayerId++,
+			InPainter.SectionGeometry.ToPaintGeometry(FVector2D(DrawOffset, 0.0f), FVector2D(1.0f, SectionSize.Y)),
+			GenericBrush,
+			ESlateDrawEffect::None,
+			FLinearColor::Gray
+		);
+
+		DrawOffset += MediaSizeX;
+	}
+}
+
+
+void FMediaThumbnailSection::DrawSampleStates(FSequencerSectionPainter& InPainter, FTimespan MediaDuration, FVector2D SectionSize, const TRangeSet<FTimespan>& RangeSet, const FLinearColor& Color) const
+{
+	static const FSlateBrush* GenericBrush = FCoreStyle::Get().GetBrush("GenericWhiteBox");
+
+	FFrameRate FrameResolution = Section->GetTypedOuter<UMovieScene>()->GetFrameResolution();
+	double SectionDuration = FFrameTime(MovieScene::DiscreteSize(Section->GetRange())) / FrameResolution;
+	const float MediaSizeX = MediaDuration.GetTotalSeconds() * SectionSize.X / SectionDuration;
+
+	TArray<TRange<FTimespan>> Ranges;
+	RangeSet.GetRanges(Ranges);
+
+	for (auto& Range : Ranges)
+	{
+		const float DrawOffset = FMath::RoundToNegativeInfinity(FTimespan::Ratio(Range.GetLowerBoundValue(), MediaDuration) * MediaSizeX);
+		const float DrawSize = FMath::RoundToPositiveInfinity(FTimespan::Ratio(Range.Size<FTimespan>(), MediaDuration) * MediaSizeX);
+		const float BarHeight = 4.0f;
+
+		FSlateDrawElement::MakeBox(
+			InPainter.DrawElements,
+			InPainter.LayerId++,
+			InPainter.SectionGeometry.ToPaintGeometry(FVector2D(DrawOffset, SectionSize.Y - BarHeight - 1.0f), FVector2D(DrawSize, BarHeight)),
+			GenericBrush,
+			ESlateDrawEffect::None,
+			Color
+		);
+	}
+}
+
+
+UMediaPlayer* FMediaThumbnailSection::GetTemplateMediaPlayer() const
+{
+	// locate the track that evaluates this section
+	if (!SectionPtr.IsValid())
+	{
+		return nullptr;
+	}
+
+	TSharedPtr<ISequencer> Sequencer = SequencerPtr.Pin();
+
+	if (!Sequencer.IsValid())
+	{
+		return nullptr; // no movie scene player
+	}
+
+	const FMovieSceneSequenceID SequenceId = Sequencer->GetFocusedTemplateID();
+	FMovieSceneEvaluationTemplate* Template = Sequencer->GetEvaluationTemplate().FindTemplate(SequenceId);
+
+	if (Template == nullptr)
+	{
+		return nullptr; // section template not found
+	}
+
+	auto OwnerTrack = Cast<UMovieSceneTrack>(SectionPtr->GetOuter());
+
+	if (OwnerTrack == nullptr)
+	{
+		return nullptr; // media track not found
+	}
+
+	const FMovieSceneTrackIdentifier TrackIdentifier = Template->GetLedger().FindTrack(OwnerTrack->GetSignature());
+	FMovieSceneEvaluationTrack* EvaluationTrack = Template->FindTrack(TrackIdentifier);
+
+	if (EvaluationTrack == nullptr)
+	{
+		return nullptr; // evaluation track not found
+	}
+
+	FMovieSceneMediaData* MediaData = nullptr;
+
+	// find the persistent data of the section being drawn
+	TArrayView<FMovieSceneEvalTemplatePtr> Children = EvaluationTrack->GetChildTemplates();
+	FPersistentEvaluationData PersistentData(*Sequencer.Get());
+
+	for (int32 ChildIndex = 0; ChildIndex < Children.Num(); ++ChildIndex)
+	{
+		if (Children[ChildIndex]->GetSourceSection() == SectionPtr)
 		{
-			MediaPlayer->OpenSource(MediaSource);
+			FMovieSceneEvaluationKey SectionKey(SequenceId, TrackIdentifier, ChildIndex);
+			PersistentData.SetSectionKey(SectionKey);
+			MediaData = PersistentData.FindSectionData<FMovieSceneMediaData>();
+
+			break;
 		}
 	}
 
-	MediaPlayer->Pause();
+	// get the template's media player
+	if (MediaData == nullptr)
+	{
+		return nullptr; // section persistent data not found
+	}
+
+	return MediaData->GetMediaPlayer();
 }
 
 
diff --git a/Engine/Plugins/Media/MediaCompositing/Source/MediaCompositingEditor/Private/Sequencer/MediaThumbnailSection.h b/Engine/Plugins/Media/MediaCompositing/Source/MediaCompositingEditor/Private/Sequencer/MediaThumbnailSection.h
--- a/Engine/Plugins/Media/MediaCompositing/Source/MediaCompositingEditor/Private/Sequencer/MediaThumbnailSection.h
+++ b/Engine/Plugins/Media/MediaCompositing/Source/MediaCompositingEditor/Private/Sequencer/MediaThumbnailSection.h
@@ -2,6 +2,8 @@
 
 #pragma once
 
+#include "Math/RangeSet.h"
+#include "Misc/Timespan.h"
 #include "Sections/ThumbnailSection.h"
 #include "Templates/SharedPointer.h"
 #include "TrackEditorThumbnail/TrackEditorThumbnail.h"
@@ -46,13 +48,13 @@ public:
 
 	//~ FThumbnailSection interface
 
-	virtual FText GetSectionTitle() const override;
-	virtual void SetSingleTime(float GlobalTime) override;
 	virtual TSharedRef<SWidget> GenerateSectionWidget() override;
-	virtual float GetSectionHeight() const override;
 	virtual FMargin GetContentPadding() const override;
-	virtual void Tick(const FGeometry& AllottedGeometry, const FGeometry& ClippedGeometry, const double InCurrentTime, const float InDeltaTime) override;
+	virtual float GetSectionHeight() const override;
+	virtual FText GetSectionTitle() const override;
+	virtual void SetSingleTime(double GlobalTime) override;
 	virtual int32 OnPaintSection(FSequencerSectionPainter& InPainter) const override;
+	virtual void Tick(const FGeometry& AllottedGeometry, const FGeometry& ClippedGeometry, const double InCurrentTime, const float InDeltaTime) override;
 
 public:
 
@@ -61,11 +63,45 @@ public:
 	virtual void Draw(FTrackEditorThumbnail& TrackEditorThumbnail) override;
 	virtual void Setup() override;
 
+protected:
+
+	/**
+	 * Draw the section's film border decoration.
+	 *
+	 * @param InPainter The object that paints the geometry.
+	 * @param SectionSize The size of the section (in local coordinates).
+	 */
+	void DrawFilmBorder(FSequencerSectionPainter& InPainter, FVector2D SectionSize) const;
+
+	/**
+	 * Draw indicators for where the media source is looping.
+	 *
+	 * @param InPainter The object that paints the geometry.
+	 * @param SectionSize The size of the section (in local coordinates).
+	 */
+	void DrawLoopIndicators(FSequencerSectionPainter& InPainter, FTimespan MediaDuration, FVector2D SectionSize) const;
+
+	/** Draw the caching state of the given media samples. */
+	void DrawSampleStates(FSequencerSectionPainter& InPainter, FTimespan MediaDuration, FVector2D SectionSize, const TRangeSet<FTimespan>& RangeSet, const FLinearColor& Color) const;
+
+	/**
+	 * Get the media player that is used by the evaluation template.
+	 *
+	 * @return The media player, or nullptr if not found.
+	 */
+	UMediaPlayer* GetTemplateMediaPlayer() const;
+
 private:
 
 	/** Internal media player used to generate the thumbnail images. */
-	UMediaPlayer* MediaPlayer;
+//	UMediaPlayer* MediaPlayer;
 
 	/** Media texture that receives the thumbnail image frames. */
-	UMediaTexture* MediaTexture;
+//	UMediaTexture* MediaTexture;
+
+	/** The section object that owns this section. */
+	TWeakObjectPtr<UMovieSceneMediaSection> SectionPtr;
+
+	/** The sequencer object that owns this section. */
+	TWeakPtr<ISequencer> SequencerPtr;
 };
diff --git a/Engine/Plugins/Media/MediaCompositing/Source/MediaCompositingEditor/Private/Sequencer/MediaTrackEditor.cpp b/Engine/Plugins/Media/MediaCompositing/Source/MediaCompositingEditor/Private/Sequencer/MediaTrackEditor.cpp
--- a/Engine/Plugins/Media/MediaCompositing/Source/MediaCompositingEditor/Private/Sequencer/MediaTrackEditor.cpp
+++ b/Engine/Plugins/Media/MediaCompositing/Source/MediaCompositingEditor/Private/Sequencer/MediaTrackEditor.cpp
@@ -4,11 +4,11 @@
 
 #include "ContentBrowserModule.h"
 #include "EditorStyleSet.h"
+#include "Framework/Application/SlateApplication.h"
 #include "Framework/MultiBox/MultiBoxBuilder.h"
 #include "IContentBrowserSingleton.h"
 #include "ISequencer.h"
 #include "ISequencerObjectChangeListener.h"
-#include "MediaPlayer.h"
 #include "MediaSource.h"
 #include "MediaTexture.h"
 #include "MovieSceneBinding.h"
@@ -21,6 +21,7 @@
 #include "TrackEditorThumbnail/TrackEditorThumbnailPool.h"
 #include "Widgets/SBoxPanel.h"
 #include "Widgets/Layout/SBox.h"
+#include "QualifiedFrameTime.h"
 
 #include "MediaThumbnailSection.h"
 
@@ -42,21 +43,13 @@ TArray<FAnimatedPropertyKey, TInlineAllocator<1>> FMediaTrackEditor::GetAnimated
 
 FMediaTrackEditor::FMediaTrackEditor(TSharedRef<ISequencer> InSequencer)
 	: FMovieSceneTrackEditor(InSequencer)
-	, PropertyKey(FAnimatedPropertyKey::FromObjectType(UMediaTexture::StaticClass()))
 {
 	ThumbnailPool = MakeShared<FTrackEditorThumbnailPool>(InSequencer);
-	OnPropertyChangedHandle = InSequencer->GetObjectChangeListener().GetOnAnimatablePropertyChanged(PropertyKey).AddRaw(this, &FMediaTrackEditor::OnAnimatedPropertyChanged);
 }
 
 
 FMediaTrackEditor::~FMediaTrackEditor()
 {
-	TSharedPtr<ISequencer> SequencerPtr = GetSequencer();
-
-	if (SequencerPtr.IsValid())
-	{
-		SequencerPtr->GetObjectChangeListener().GetOnAnimatablePropertyChanged(PropertyKey).Remove(OnPropertyChangedHandle);
-	}
 }
 
 
@@ -66,148 +59,30 @@ FMediaTrackEditor::~FMediaTrackEditor()
 UMovieSceneTrack* FMediaTrackEditor::AddTrack(UMovieScene* FocusedMovieScene, const FGuid& ObjectHandle, TSubclassOf<class UMovieSceneTrack> TrackClass, FName UniqueTypeName)
 {
 	UMovieSceneTrack* Track = FocusedMovieScene->AddTrack(TrackClass, ObjectHandle);
-	if (UMovieSceneMediaTrack* MediaTrack = Cast<UMovieSceneMediaTrack>(Track))
-	{
-		MediaTrack->UniqueTrackName = UniqueTypeName;
-	}
+	UMovieSceneMediaTrack* MediaTrack = Cast<UMovieSceneMediaTrack>(Track);
+
 	return Track;
 }
 
-void FMediaTrackEditor::OnAnimatedPropertyChanged(const FPropertyChangedParams& PropertyChangedParams)
-{
-	const UProperty* ChangedProperty = PropertyChangedParams.PropertyPath.GetLeafMostProperty().Property.Get();
-	if (!ChangedProperty)
-	{
-		return;
-	}
-
-	FText DisplayText = ChangedProperty->GetDisplayNameText();
-	FName UniqueName(*PropertyChangedParams.PropertyPath.ToString(TEXT(".")));
 
-	// Set up the appropriate name for the track from an array index if necessary
-	for (int32 PropertyIndex = PropertyChangedParams.PropertyPath.GetNumProperties() - 1; PropertyIndex >= 0; --PropertyIndex)
-	{
-		const FPropertyInfo& Info = PropertyChangedParams.PropertyPath.GetPropertyInfo(PropertyIndex);
-		const UArrayProperty* ParentArrayProperty = PropertyIndex > 0 ? Cast<UArrayProperty>(PropertyChangedParams.PropertyPath.GetPropertyInfo(PropertyIndex-1).Property.Get()) : nullptr;
-
-		UProperty* Property = Info.Property.Get();
-		if (Property && Info.ArrayIndex != INDEX_NONE)
-		{
-			DisplayText = FText::Format(LOCTEXT("MediaTrackNameFormat", "{0} ({1}[{2}])"),
-				ChangedProperty->GetDisplayNameText(),
-				(ParentArrayProperty ? ParentArrayProperty : Property)->GetDisplayNameText(),
-				FText::AsNumber(Info.ArrayIndex)
-				);
-			break;
-		}
-	}
-
-	for (UObject* Object : PropertyChangedParams.ObjectsThatChanged)
-	{
-		FFindOrCreateHandleResult HandleResult = FindOrCreateHandleToObject(Object);
-		if (!ensure(HandleResult.Handle.IsValid()))
-		{
-			continue;
-		}
-		
-		FFindOrCreateTrackResult TrackResult = FindOrCreateTrackForObject(HandleResult.Handle, UMovieSceneMediaTrack::StaticClass(), UniqueName);
-		UMovieSceneMediaTrack* MediaTrack = Cast<UMovieSceneMediaTrack>(TrackResult.Track);
-		
-		if (TrackResult.bWasCreated)
-		{
-			MediaTrack->SetPropertyNameAndPath(ChangedProperty->GetFName(), PropertyChangedParams.PropertyPath.ToString(TEXT(".")));
-			MediaTrack->SetDisplayName(DisplayText);
-		}
-
-		GetSequencer()->NotifyMovieSceneDataChanged(EMovieSceneDataChangeType::MovieSceneStructureItemAdded);
-	}
-}
-
-void FMediaTrackEditor::BuildObjectBindingTrackMenu(FMenuBuilder& MenuBuilder, const FGuid& ObjectBinding, const UClass* ObjectClass)
+void FMediaTrackEditor::BuildAddTrackMenu(FMenuBuilder& MenuBuilder)
 {
-	// We only know how to add specific properties for image planes and components. Anything else must be keyed through the generic MediaPlayer property
-	if (!ObjectClass->IsChildOf(AMediaPlane::StaticClass()) && !ObjectClass->IsChildOf(UMediaPlaneComponent::StaticClass()))
-	{
-		return;
-	}
-
-	// Find the spawned object or its template
-	UObject* Object = GetSequencer()->FindSpawnedObjectOrTemplate(ObjectBinding);
-	if (!Object)
-	{
-		return;
-	}
-
-	auto GeneratePropertyPath = [this](UMediaPlaneComponent* MediaPlaneComponent)
-	{
-		check(MediaPlaneComponent);
-
-		UStructProperty* MediaPlanesProperty = MediaPlaneComponent->GetMediaPlaneProperty();
-		UProperty* MediaTextureProperty = FindField<UProperty>(FMediaPlaneParameters::StaticStruct(), GET_MEMBER_NAME_CHECKED(FMediaPlaneParameters, MediaTexture));
-
-		check(MediaPlanesProperty);
-		check(MediaTextureProperty);
-
-		TSharedRef<FPropertyPath> Path = FPropertyPath::CreateEmpty();
-		Path->AddProperty(FPropertyInfo(MediaPlanesProperty));
-		Path->AddProperty(FPropertyInfo(MediaTextureProperty));
-
-		return Path;
-	};
-
-	auto AddNewTrack = [this](UMediaPlaneComponent* MediaPlaneComponent, TSharedRef<FPropertyPath> PropertyPath)
-	{
-		check(MediaPlaneComponent);
-
-		FPropertyChangedParams ChangedParams(TArray<UObject*>({ MediaPlaneComponent }), *PropertyPath, NAME_None, ESequencerKeyMode::ManualKeyForced);
-		this->OnAnimatedPropertyChanged(ChangedParams);
-	};
-
-	// Try and root out an image plane component
-	UMediaPlaneComponent* Component = Cast<UMediaPlaneComponent>(Object);
-	if (AMediaPlane* MediaPlane = Cast<AMediaPlane>(Object))
-	{
-		Component = Cast<UMediaPlaneComponent>(MediaPlane->GetRootComponent());
-	}
-
-	if (Component)
-	{
-		UMovieScene* MovieScene = GetSequencer()->GetFocusedMovieSceneSequence()->GetMovieScene();
-
-		TSharedRef<FPropertyPath> Path = GeneratePropertyPath(Component);
-
-		bool bCanAddTrack = true;
-
-		const FMovieSceneBinding* Binding = MovieScene->GetBindings().FindByPredicate([=](const FMovieSceneBinding& InBinding){ return InBinding.GetObjectGuid() == ObjectBinding; });
-		if (Binding)
-		{
-			FString PredicatePath = Path->ToString(TEXT("."));
-
-			bCanAddTrack = !Binding->GetTracks().ContainsByPredicate(
-				[&](UMovieSceneTrack* Track)
-				{
-					UMovieSceneMediaTrack* MediaTrack = Cast<UMovieSceneMediaTrack>(Track);
-					return MediaTrack && MediaTrack->GetPropertyPath() == PredicatePath;
-				}
-			);
-		}
-
-		MenuBuilder.AddMenuEntry(
-			LOCTEXT("AddMediaTrack_Text", "Media"),
-			LOCTEXT("AddMediaTrack_Tip", "Adds a media track that controls the media presented to the media plane."),
-			FSlateIcon(),
-			FUIAction(
-				FExecuteAction::CreateLambda(AddNewTrack, Component, Path),
-				FCanExecuteAction::CreateLambda([=]{ return bCanAddTrack; })
-			)
-		);
-	}
+	MenuBuilder.AddMenuEntry(
+		LOCTEXT("AddTrack", "Media Track"),
+		LOCTEXT("AddTooltip", "Adds a new master media track that can play media sources."),
+		FSlateIcon(FEditorStyle::GetStyleSetName(), "Sequencer.Tracks.Media"),
+		FUIAction(
+			FExecuteAction::CreateRaw(this, &FMediaTrackEditor::HandleAddMediaTrackMenuEntryExecute)
+		)
+	);
 }
 
+
 TSharedPtr<SWidget> FMediaTrackEditor::BuildOutlinerEditWidget(const FGuid& ObjectBinding, UMovieSceneTrack* Track, const FBuildEditWidgetParams& Params)
 {
 	UMovieSceneMediaTrack* MediaTrack = Cast<UMovieSceneMediaTrack>(Track);
-	if (!MediaTrack)
+
+	if (MediaTrack == nullptr)
 	{
 		return SNullWidget::NullWidget;
 	}
@@ -215,12 +90,13 @@ TSharedPtr<SWidget> FMediaTrackEditor::BuildOutlinerEditWidget(const FGuid& Obje
 	auto CreatePicker = [this, MediaTrack]
 	{
 		FAssetPickerConfig AssetPickerConfig;
-		AssetPickerConfig.OnAssetSelected = FOnAssetSelected::CreateRaw(this, &FMediaTrackEditor::AddNewSection, MediaTrack);
-		AssetPickerConfig.bAllowNullSelection = false;
-		AssetPickerConfig.InitialAssetViewType = EAssetViewType::List;
-
-		AssetPickerConfig.Filter.bRecursiveClasses = true;
-		AssetPickerConfig.Filter.ClassNames.Add(UMediaSource::StaticClass()->GetFName());
+		{
+			AssetPickerConfig.OnAssetSelected = FOnAssetSelected::CreateRaw(this, &FMediaTrackEditor::AddNewSection, MediaTrack);
+			AssetPickerConfig.bAllowNullSelection = false;
+			AssetPickerConfig.InitialAssetViewType = EAssetViewType::List;
+			AssetPickerConfig.Filter.bRecursiveClasses = true;
+			AssetPickerConfig.Filter.ClassNames.Add(UMediaSource::StaticClass()->GetFName());
+		}
 
 		FContentBrowserModule& ContentBrowserModule = FModuleManager::Get().LoadModuleChecked<FContentBrowserModule>(TEXT("ContentBrowser"));
 
@@ -235,7 +111,8 @@ TSharedPtr<SWidget> FMediaTrackEditor::BuildOutlinerEditWidget(const FGuid& Obje
 	};
 
 	return SNew(SHorizontalBox)
-		+ SHorizontalBox::Slot()
+
+	+ SHorizontalBox::Slot()
 		.AutoWidth()
 		.VAlign(VAlign_Center)
 		[
@@ -243,80 +120,176 @@ TSharedPtr<SWidget> FMediaTrackEditor::BuildOutlinerEditWidget(const FGuid& Obje
 		];
 }
 
-void FMediaTrackEditor::AddNewSection(const FAssetData& AssetData, UMovieSceneMediaTrack* Track)
+
+bool FMediaTrackEditor::HandleAssetAdded(UObject* Asset, const FGuid& TargetObjectGuid)
 {
-	UMediaSource* MediaSource = Cast<UMediaSource>(AssetData.GetAsset());
-	if (!MediaSource)
+	if ((Asset == nullptr) || !Asset->IsA<UMediaSource>())
 	{
-		return;
+		return false;
 	}
 
-	check(Track);
-
-	float TimeToStart = GetSequencer()->GetLocalTime();
-
-	UMediaPlayer* TransientPlayer = NewObject<UMediaPlayer>(GetTransientPackage(), MakeUniqueObjectName(GetTransientPackage(), UMediaPlayer::StaticClass()));
-	TransientPlayer->AddToRoot();
-
-	TRange<float> SectionRange(TimeToStart, TimeToStart + 1.f);
+	auto MediaSource = Cast<UMediaSource>(Asset);
 
-	if (TransientPlayer->OpenSource(MediaSource))
+	if (TargetObjectGuid.IsValid())
 	{
-		// @todo: Make this more robust for all media types
-		FTimespan Length = TransientPlayer->GetDuration();
-		if (Length.GetTotalSeconds() > 0)
+		TArray<TWeakObjectPtr<>> OutObjects;
+
+		for (TWeakObjectPtr<> Object : GetSequencer()->FindObjectsInCurrentSequence(TargetObjectGuid))
 		{
-			SectionRange = TRange<float>(TimeToStart, TRangeBound<float>::Inclusive(TimeToStart + float(Length.GetTicks()) / ETimespan::TicksPerSecond));
+			OutObjects.Add(Object);
 		}
 
-		TransientPlayer->Close();
+		int32 RowIndex = INDEX_NONE;
+		AnimatablePropertyChanged(FOnKeyProperty::CreateRaw(this, &FMediaTrackEditor::AddAttachedMediaSource, MediaSource, OutObjects, RowIndex));
+	}
+	else
+	{
+		int32 RowIndex = INDEX_NONE;
+		AnimatablePropertyChanged(FOnKeyProperty::CreateRaw(this, &FMediaTrackEditor::AddMasterMediaSource, MediaSource, RowIndex));
 	}
 
-	TransientPlayer->RemoveFromRoot();
+	return true;
+}
+
+
+TSharedRef<ISequencerSection> FMediaTrackEditor::MakeSectionInterface(UMovieSceneSection& SectionObject, UMovieSceneTrack& Track, FGuid ObjectBinding)
+{
+	check(SupportsType(SectionObject.GetOuter()->GetClass()));
+	return MakeShared<FMediaThumbnailSection>(*CastChecked<UMovieSceneMediaSection>(&SectionObject), ThumbnailPool, GetSequencer());
+}
+
+
+bool FMediaTrackEditor::SupportsType(TSubclassOf<UMovieSceneTrack> TrackClass) const
+{
+	return TrackClass.Get() && TrackClass.Get()->IsChildOf(UMovieSceneMediaTrack::StaticClass());
+}
+
+
+void FMediaTrackEditor::Tick(float DeltaTime)
+{
+	ThumbnailPool->DrawThumbnails();
+}
 
-	// Find a row to put this video on
-	TArray<int32> InvalidRows;
-	for (UMovieSceneSection* Section : Track->GetAllSections())
+
+const FSlateBrush* FMediaTrackEditor::GetIconBrush() const
+{
+	return FEditorStyle::GetBrush("Sequencer.Tracks.Media");
+}
+
+
+/* FMediaTrackEditor implementation
+ *****************************************************************************/
+
+FKeyPropertyResult FMediaTrackEditor::AddAttachedMediaSource(FFrameNumber KeyTime, UMediaSource* MediaSource, TArray<TWeakObjectPtr<UObject>> ObjectsToAttachTo, int32 RowIndex)
+{
+	FKeyPropertyResult KeyPropertyResult;
+
+	for (int32 ObjectIndex = 0; ObjectIndex < ObjectsToAttachTo.Num(); ++ObjectIndex)
 	{
-		if (Section->GetRange().Overlaps(SectionRange))
+		UObject* Object = ObjectsToAttachTo[ObjectIndex].Get();
+
+		FFindOrCreateHandleResult HandleResult = FindOrCreateHandleToObject(Object);
+		FGuid ObjectHandle = HandleResult.Handle;
+		KeyPropertyResult.bHandleCreated |= HandleResult.bWasCreated;
+
+		if (ObjectHandle.IsValid())
 		{
-			InvalidRows.Add(Section->GetRowIndex());
+			FFindOrCreateTrackResult TrackResult = FindOrCreateTrackForObject(ObjectHandle, UMovieSceneMediaTrack::StaticClass());
+			UMovieSceneTrack* Track = TrackResult.Track;
+			KeyPropertyResult.bTrackCreated |= TrackResult.bWasCreated;
+
+			if (ensure(Track))
+			{
+				auto MediaTrack = Cast<UMovieSceneMediaTrack>(Track);
+				UMovieSceneSection* NewSection = MediaTrack->AddNewMediaSourceOnRow(*MediaSource, KeyTime, RowIndex);
+				MediaTrack->SetDisplayName(LOCTEXT("MediaTrackName", "Media"));
+				KeyPropertyResult.bTrackModified = true;
+
+				GetSequencer()->EmptySelection();
+				GetSequencer()->SelectSection(NewSection);
+				GetSequencer()->ThrobSectionSelection();
+			}
 		}
 	}
 
-	int32 BestRow = 0;
-	for ( ; ; ++BestRow)
+	return KeyPropertyResult;
+}
+
+
+FKeyPropertyResult FMediaTrackEditor::AddMasterMediaSource(FFrameNumber KeyTime, UMediaSource* MediaSource, int32 RowIndex)
+{
+	FKeyPropertyResult KeyPropertyResult;
+
+	FFindOrCreateMasterTrackResult<UMovieSceneMediaTrack> TrackResult = FindOrCreateMasterTrack<UMovieSceneMediaTrack>();
+	UMovieSceneTrack* Track = TrackResult.Track;
+	auto MediaTrack = Cast<UMovieSceneMediaTrack>(Track);
+
+	MediaTrack->AddNewMediaSourceOnRow(*MediaSource, KeyTime, RowIndex);
+
+	if (TrackResult.bWasCreated)
 	{
-		if (!InvalidRows.Contains(BestRow))
-		{
-			break;
-		}
+		MediaTrack->SetDisplayName(LOCTEXT("MediaTrackName", "Media"));
 	}
 
-	UMovieSceneMediaSection* Section = CastChecked<UMovieSceneMediaSection>(Track->CreateNewSection());
-	Section->SetRange(SectionRange);
-	Section->SetRowIndex(BestRow);
-	Section->SetMediaSource(MediaSource);
-	Track->AddSection(*Section);
+	KeyPropertyResult.bTrackModified = true;
 
-	GetSequencer()->NotifyMovieSceneDataChanged(EMovieSceneDataChangeType::MovieSceneStructureItemAdded);
+	return KeyPropertyResult;
 }
 
-void FMediaTrackEditor::Tick(float DeltaTime)
-{
-	ThumbnailPool->DrawThumbnails();
-}
 
-TSharedRef<ISequencerSection> FMediaTrackEditor::MakeSectionInterface(UMovieSceneSection& SectionObject, UMovieSceneTrack& Track, FGuid ObjectBinding)
+void FMediaTrackEditor::AddNewSection(const FAssetData& AssetData, UMovieSceneMediaTrack* Track)
 {
-	check(SupportsType(SectionObject.GetOuter()->GetClass()));
-	return MakeShared<FMediaThumbnailSection>(*CastChecked<UMovieSceneMediaSection>(&SectionObject), ThumbnailPool, GetSequencer());
+	FSlateApplication::Get().DismissAllMenus();
+
+	UObject* SelectedObject = AssetData.GetAsset();
+
+	if (SelectedObject)
+	{
+		auto MediaSource = CastChecked<UMediaSource>(AssetData.GetAsset());
+
+		if (MediaSource != nullptr)
+		{
+			const FScopedTransaction Transaction(NSLOCTEXT("Sequencer", "AddMedia_Transaction", "Add Media"));
+
+			auto MediaTrack = Cast<UMovieSceneMediaTrack>(Track);
+			MediaTrack->Modify();
+
+			FFrameTime KeyTime = GetSequencer()->GetLocalTime().Time;
+			UMovieSceneSection* NewSection = MediaTrack->AddNewMediaSource(*MediaSource, KeyTime.FrameNumber);
+
+			GetSequencer()->EmptySelection();
+			GetSequencer()->SelectSection(NewSection);
+			GetSequencer()->ThrobSectionSelection();
+
+			GetSequencer()->NotifyMovieSceneDataChanged(EMovieSceneDataChangeType::MovieSceneStructureItemAdded);
+		}
+	}
 }
 
-bool FMediaTrackEditor::SupportsType(TSubclassOf<UMovieSceneTrack> TrackClass) const
+
+/* FMediaTrackEditor callbacks
+ *****************************************************************************/
+
+void FMediaTrackEditor::HandleAddMediaTrackMenuEntryExecute()
 {
-	return TrackClass.Get() && TrackClass.Get()->IsChildOf(UMovieSceneMediaTrack::StaticClass());
+	UMovieScene* FocusedMovieScene = GetFocusedMovieScene();
+
+	if (FocusedMovieScene == nullptr)
+	{
+		return;
+	}
+
+	const FScopedTransaction Transaction(NSLOCTEXT("Sequencer", "AddMediaTrack_Transaction", "Add Media Track"));
+	FocusedMovieScene->Modify();
+	
+	auto NewTrack = FocusedMovieScene->AddMasterTrack<UMovieSceneMediaTrack>();
+	ensure(NewTrack);
+
+	NewTrack->SetDisplayName(LOCTEXT("MediaTrackName", "Media"));
+
+	GetSequencer()->NotifyMovieSceneDataChanged(EMovieSceneDataChangeType::MovieSceneStructureItemAdded);
 }
 
 
+
 #undef LOCTEXT_NAMESPACE
diff --git a/Engine/Plugins/Media/MediaCompositing/Source/MediaCompositingEditor/Private/Sequencer/MediaTrackEditor.h b/Engine/Plugins/Media/MediaCompositing/Source/MediaCompositingEditor/Private/Sequencer/MediaTrackEditor.h
--- a/Engine/Plugins/Media/MediaCompositing/Source/MediaCompositingEditor/Private/Sequencer/MediaTrackEditor.h
+++ b/Engine/Plugins/Media/MediaCompositing/Source/MediaCompositingEditor/Private/Sequencer/MediaTrackEditor.h
@@ -55,22 +55,31 @@ public:
 
 	//~ FMovieSceneTrackEditor interface
 
+	virtual UMovieSceneTrack* AddTrack(UMovieScene* FocusedMovieScene, const FGuid& ObjectHandle, TSubclassOf<class UMovieSceneTrack> TrackClass, FName UniqueTypeName) override;
+	virtual void BuildAddTrackMenu(FMenuBuilder& MenuBuilder) override;
+	virtual TSharedPtr<SWidget> BuildOutlinerEditWidget(const FGuid& ObjectBinding, UMovieSceneTrack* Track, const FBuildEditWidgetParams& Params) override;
+	virtual bool HandleAssetAdded(UObject* Asset, const FGuid& TargetObjectGuid) override;
 	virtual TSharedRef<ISequencerSection> MakeSectionInterface(UMovieSceneSection& SectionObject, UMovieSceneTrack& Track, FGuid ObjectBinding) override;
 	virtual bool SupportsType(TSubclassOf<UMovieSceneTrack> TrackClass) const override;
 	virtual void Tick(float DeltaTime) override;
-	virtual void BuildObjectBindingTrackMenu(FMenuBuilder& MenuBuilder, const FGuid& ObjectBinding, const UClass* ObjectClass) override;
-	virtual TSharedPtr<SWidget> BuildOutlinerEditWidget(const FGuid& ObjectBinding, UMovieSceneTrack* Track, const FBuildEditWidgetParams& Params) override;
-	virtual UMovieSceneTrack* AddTrack(UMovieScene* FocusedMovieScene, const FGuid& ObjectHandle, TSubclassOf<class UMovieSceneTrack> TrackClass, FName UniqueTypeName) override;
+	virtual const FSlateBrush* GetIconBrush() const override;
 
 protected:
 
-	void OnAnimatedPropertyChanged(const FPropertyChangedParams& PropertyChangedParams);
+	/** Callback for AnimatablePropertyChanged in HandleAssetAdded for attached media sources. */
+	FKeyPropertyResult AddAttachedMediaSource(FFrameNumber KeyTime, class UMediaSource* MediaSource, TArray<TWeakObjectPtr<UObject>> ObjectsToAttachTo, int32 RowIndex);
+
+	/** Callback for AnimatablePropertyChanged in HandleAssetAdded for master media sources. */
+	FKeyPropertyResult AddMasterMediaSource(FFrameNumber KeyTime, class UMediaSource* MediaSource, int32 RowIndex);
 
 	void AddNewSection(const FAssetData& Asset, UMovieSceneMediaTrack* Track);
 
 private:
 
-	FAnimatedPropertyKey PropertyKey;
-	FDelegateHandle OnPropertyChangedHandle;
+	/** Callback for executing the "Add Media Track" menu entry. */
+	void HandleAddMediaTrackMenuEntryExecute();
+
+private:
+
 	TSharedPtr<FTrackEditorThumbnailPool> ThumbnailPool;
 };
diff --git a/Engine/Plugins/Media/MediaPlayerEditor/Source/MediaPlayerEditor/Private/Widgets/SMediaPlayerEditorCache.cpp b/Engine/Plugins/Media/MediaPlayerEditor/Source/MediaPlayerEditor/Private/Widgets/SMediaPlayerEditorCache.cpp
--- a/Engine/Plugins/Media/MediaPlayerEditor/Source/MediaPlayerEditor/Private/Widgets/SMediaPlayerEditorCache.cpp
+++ b/Engine/Plugins/Media/MediaPlayerEditor/Source/MediaPlayerEditor/Private/Widgets/SMediaPlayerEditorCache.cpp
@@ -75,7 +75,8 @@ int32 SMediaPlayerEditorCache::OnPaint(const FPaintArgs& Args, const FGeometry&
 
 void SMediaPlayerEditorCache::DrawPlayerPosition(FTimespan Time, const FGeometry& AllottedGeometry, FSlateWindowElementList& OutDrawElements, int32 LayerId, const FWidgetStyle& InWidgetStyle, const FLinearColor& Color) const
 {
-	const FSlateBrush* GenericBrush = FCoreStyle::Get().GetBrush("GenericWhiteBox");
+	static const FSlateBrush* GenericBrush = FCoreStyle::Get().GetBrush("GenericWhiteBox");
+
 	const FTimespan Duration = MediaPlayer->GetDuration();
 	const float MarkerSize = PositionMarkerSize.Get();
 
@@ -119,7 +120,8 @@ void SMediaPlayerEditorCache::DrawSampleCache(EMediaTrackType TrackType, const F
 
 void SMediaPlayerEditorCache::DrawSampleStates(const TRangeSet<FTimespan>& RangeSet, const FGeometry& AllottedGeometry, FSlateWindowElementList& OutDrawElements, int32 LayerId, const FWidgetStyle& InWidgetStyle, const FLinearColor& Color, float YPos, float YScale) const
 {
-	const FSlateBrush* GenericBrush = FCoreStyle::Get().GetBrush("GenericWhiteBox");
+	static const FSlateBrush* GenericBrush = FCoreStyle::Get().GetBrush("GenericWhiteBox");
+
 	const FTimespan Duration = MediaPlayer->GetDuration();
 
 	FSlateClippingZone ClippingZone(AllottedGeometry);
diff --git a/Engine/Plugins/Media/MediaPlayerEditor/Source/MediaPlayerEditor/Private/Widgets/SMediaPlayerEditorCache.h b/Engine/Plugins/Media/MediaPlayerEditor/Source/MediaPlayerEditor/Private/Widgets/SMediaPlayerEditorCache.h
--- a/Engine/Plugins/Media/MediaPlayerEditor/Source/MediaPlayerEditor/Private/Widgets/SMediaPlayerEditorCache.h
+++ b/Engine/Plugins/Media/MediaPlayerEditor/Source/MediaPlayerEditor/Private/Widgets/SMediaPlayerEditorCache.h
@@ -69,7 +69,7 @@ protected:
 	void DrawSampleCache(EMediaTrackType TrackType, const FGeometry& AllottedGeometry, FSlateWindowElementList& OutDrawElements, int32 LayerId, const FWidgetStyle& InWidgetStyle, float YPos, float YScale) const;
 	
 	/** Draw the caching state of the given media samples. */
-	void DrawSampleStates(const TRangeSet<FTimespan>& Ranges, const FGeometry& AllottedGeometry, FSlateWindowElementList& OutDrawElements, int32 LayerId, const FWidgetStyle& InWidgetStyle, const FLinearColor& Color, float YPos, float YScale) const;
+	void DrawSampleStates(const TRangeSet<FTimespan>& RangeSet, const FGeometry& AllottedGeometry, FSlateWindowElementList& OutDrawElements, int32 LayerId, const FWidgetStyle& InWidgetStyle, const FLinearColor& Color, float YPos, float YScale) const;
 
 private:
 
diff --git a/Engine/Plugins/Messaging/UdpMessaging/Source/UdpMessaging/Private/Transport/UdpMessageBeacon.cpp b/Engine/Plugins/Messaging/UdpMessaging/Source/UdpMessaging/Private/Transport/UdpMessageBeacon.cpp
--- a/Engine/Plugins/Messaging/UdpMessaging/Source/UdpMessaging/Private/Transport/UdpMessageBeacon.cpp
+++ b/Engine/Plugins/Messaging/UdpMessaging/Source/UdpMessaging/Private/Transport/UdpMessageBeacon.cpp
@@ -76,6 +76,12 @@ void FUdpMessageBeacon::SetEndpointCount(int32 EndpointCount)
 /* FRunnable interface
  *****************************************************************************/
 
+FSingleThreadRunnable* FUdpMessageBeacon::GetSingleThreadInterface()
+{
+	return this;
+}
+
+
 bool FUdpMessageBeacon::Init()
 {
 	return true;
@@ -87,20 +93,11 @@ uint32 FUdpMessageBeacon::Run()
 	while (!Stopping)
 	{
 		FDateTime CurrentTime = FDateTime::UtcNow();
-
-		if (CurrentTime >= NextHelloTime)
-		{
-			// calculate the next send interval
-			BeaconInterval = FMath::Max(MinimumInterval, IntervalPerEndpoint * LastEndpointCount);
-			NextHelloTime = CurrentTime + BeaconInterval;
-
-			SendSegment(EUdpMessageSegments::Hello);
-		}
-
+		Update(CurrentTime, BeaconInterval);
 		EndpointLeftEvent->Wait(NextHelloTime - CurrentTime);
 	}
 
-	SendSegment(EUdpMessageSegments::Bye);
+	SendSegment(EUdpMessageSegments::Bye, BeaconInterval);
 
 	return 0;
 }
@@ -115,7 +112,7 @@ void FUdpMessageBeacon::Stop()
 /* FUdpMessageHelloSender implementation
  *****************************************************************************/
 
-void FUdpMessageBeacon::SendSegment(EUdpMessageSegments SegmentType)
+bool FUdpMessageBeacon::SendSegment(EUdpMessageSegments SegmentType, const FTimespan& SocketWaitTime)
 {
 	FUdpMessageSegment::FHeader Header;
 	{
@@ -132,6 +129,40 @@ void FUdpMessageBeacon::SendSegment(EUdpMessageSegments SegmentType)
 
 	int32 Sent;
 
-	Socket->Wait(ESocketWaitConditions::WaitForWrite, BeaconInterval);
-	Socket->SendTo(Writer.GetData(), Writer.Num(), Sent, *MulticastAddress);
+	if (!Socket->Wait(ESocketWaitConditions::WaitForWrite, SocketWaitTime))
+	{
+		return false; // socket not ready for sending
+	}
+
+	if (!Socket->SendTo(Writer.GetData(), Writer.Num(), Sent, *MulticastAddress))
+	{
+		return false; // send failed
+	}
+
+	return true;
+}
+
+
+void FUdpMessageBeacon::Update(const FDateTime& CurrentTime, const FTimespan& SocketWaitTime)
+{
+	if (CurrentTime < NextHelloTime)
+	{
+		return;
+	}
+
+	BeaconInterval = FMath::Max(MinimumInterval, IntervalPerEndpoint * LastEndpointCount);
+
+	if (SendSegment(EUdpMessageSegments::Hello, SocketWaitTime))
+	{
+		NextHelloTime = CurrentTime + BeaconInterval;
+	}
+}
+
+
+/* FSingleThreadRunnable interface
+ *****************************************************************************/
+
+void FUdpMessageBeacon::Tick()
+{
+	Update(FDateTime::UtcNow(), FTimespan::Zero());
 }
diff --git a/Engine/Plugins/Messaging/UdpMessaging/Source/UdpMessaging/Private/Transport/UdpMessageBeacon.h b/Engine/Plugins/Messaging/UdpMessaging/Source/UdpMessaging/Private/Transport/UdpMessageBeacon.h
--- a/Engine/Plugins/Messaging/UdpMessaging/Source/UdpMessaging/Private/Transport/UdpMessageBeacon.h
+++ b/Engine/Plugins/Messaging/UdpMessaging/Source/UdpMessaging/Private/Transport/UdpMessageBeacon.h
@@ -5,6 +5,7 @@
 #include "CoreTypes.h"
 #include "Misc/DateTime.h"
 #include "Misc/Guid.h"
+#include "Misc/SingleThreadRunnable.h"
 #include "Misc/Timespan.h"
 #include "HAL/Runnable.h"
 #include "Shared/UdpMessageSegment.h"
@@ -23,11 +24,12 @@ struct FIPv4Endpoint;
  */
 class FUdpMessageBeacon
 	: public FRunnable
+	, private FSingleThreadRunnable
 {
 public:
 
 	/** 
-	 * Creates and initializes a new Hello sender.
+	 * Creates and initializes a new beacon sender.
 	 *
 	 * @param InSocket The network socket used to send Hello segments.
 	 * @param InSocketId The network socket identifier (used to detect unicast endpoint).
@@ -61,6 +63,7 @@ public:
 
 	//~ FRunnable interface
 
+	virtual FSingleThreadRunnable* GetSingleThreadInterface() override;
 	virtual bool Init() override;
 	virtual uint32 Run() override;
 	virtual void Stop() override;
@@ -72,8 +75,24 @@ protected:
 	 * Sends the specified segment.
 	 *
 	 * @param SegmentType The type of segment to send (Hello or Bye).
+	 * @param SocketWaitTime Maximum time to wait for the socket to be ready.
+	 * @return true on success, false otherwise.
 	 */
-	void SendSegment(EUdpMessageSegments SegmentType);
+	bool SendSegment(EUdpMessageSegments SegmentType, const FTimespan& SocketWaitTime);
+
+	/**
+	 * Update the beacon sender.
+	 *
+	 * @param CurrentTime The current time (in UTC).
+	 * @param SocketWaitTime Maximum time to wait for the socket to be ready.
+	 */
+	void Update(const FDateTime& CurrentTime, const FTimespan& SocketWaitTime);
+
+protected:
+
+	//~ FSingleThreadRunnable interface
+
+	virtual void Tick() override;
 
 private:
 
diff --git a/Engine/Plugins/Messaging/UdpMessaging/Source/UdpMessaging/Private/Transport/UdpMessageProcessor.cpp b/Engine/Plugins/Messaging/UdpMessaging/Source/UdpMessaging/Private/Transport/UdpMessageProcessor.cpp
--- a/Engine/Plugins/Messaging/UdpMessaging/Source/UdpMessaging/Private/Transport/UdpMessageProcessor.cpp
+++ b/Engine/Plugins/Messaging/UdpMessaging/Source/UdpMessaging/Private/Transport/UdpMessageProcessor.cpp
@@ -108,6 +108,12 @@ bool FUdpMessageProcessor::EnqueueOutboundMessage(const TSharedRef<FUdpSerialize
 /* FRunnable interface
  *****************************************************************************/
 
+FSingleThreadRunnable* FUdpMessageProcessor::GetSingleThreadInterface()
+{
+	return this;
+}
+
+
 bool FUdpMessageProcessor::Init()
 {
 	Beacon = new FUdpMessageBeacon(Socket, LocalNodeId, MulticastEndpoint);
@@ -123,12 +129,7 @@ uint32 FUdpMessageProcessor::Run()
 	{
 		if (WorkEvent->Wait(CalculateWaitTime()))
 		{
-			CurrentTime = FDateTime::UtcNow();
-
-			ConsumeInboundSegments();
-			ConsumeOutboundMessages();
-			UpdateKnownNodes();
-			UpdateStaticNodes();
+			Update();
 		}
 	}
 	
@@ -442,6 +443,17 @@ void FUdpMessageProcessor::RemoveKnownNode(const FGuid& NodeId)
 }
 
 
+void FUdpMessageProcessor::Update()
+{
+	CurrentTime = FDateTime::UtcNow();
+
+	ConsumeInboundSegments();
+	ConsumeOutboundMessages();
+	UpdateKnownNodes();
+	UpdateStaticNodes();
+}
+
+
 void FUdpMessageProcessor::UpdateKnownNodes()
 {
 	// remove dead remote endpoints
@@ -531,3 +543,12 @@ void FUdpMessageProcessor::UpdateStaticNodes()
 		UpdateSegmenters(StaticNodePair.Value);
 	}
 }
+
+
+/* FSingleThreadRunnable interface
+ *****************************************************************************/
+
+void FUdpMessageProcessor::Tick()
+{
+	Update();
+}
diff --git a/Engine/Plugins/Messaging/UdpMessaging/Source/UdpMessaging/Private/Transport/UdpMessageProcessor.h b/Engine/Plugins/Messaging/UdpMessaging/Source/UdpMessaging/Private/Transport/UdpMessageProcessor.h
--- a/Engine/Plugins/Messaging/UdpMessaging/Source/UdpMessaging/Private/Transport/UdpMessageProcessor.h
+++ b/Engine/Plugins/Messaging/UdpMessaging/Source/UdpMessaging/Private/Transport/UdpMessageProcessor.h
@@ -11,6 +11,7 @@
 #include "Interfaces/IPv4/IPv4Endpoint.h"
 #include "Misc/DateTime.h"
 #include "Misc/Guid.h"
+#include "Misc/SingleThreadRunnable.h"
 #include "Misc/Timespan.h"
 #include "Templates/SharedPointer.h"
 
@@ -34,6 +35,7 @@ class IMessageAttachment;
  */
 class FUdpMessageProcessor
 	: public FRunnable
+	, private FSingleThreadRunnable
 {
 	/** Structure for known remote endpoints. */
 	struct FNodeInfo
@@ -203,6 +205,7 @@ public:
 
 	//~ FRunnable interface
 
+	virtual FSingleThreadRunnable* GetSingleThreadInterface() override;
 	virtual bool Init() override;
 	virtual uint32 Run() override;
 	virtual void Stop() override;
@@ -314,6 +317,9 @@ protected:
 	 */
 	void RemoveKnownNode(const FGuid& NodeId);
 
+	/** Update the message processor. */
+	void Update();
+
 	/** Updates all known remote nodes. */
 	void UpdateKnownNodes();
 
@@ -327,6 +333,12 @@ protected:
 	/** Updates all static remote nodes. */
 	void UpdateStaticNodes();
 
+protected:
+
+	//~ FSingleThreadRunnable interface
+
+	virtual void Tick() override;
+
 private:
 
 	/** Holds the queue of outbound messages. */
diff --git a/Engine/Plugins/Messaging/UdpMessaging/Source/UdpMessaging/Private/Tunnel/UdpMessageTunnel.cpp b/Engine/Plugins/Messaging/UdpMessaging/Source/UdpMessaging/Private/Tunnel/UdpMessageTunnel.cpp
--- a/Engine/Plugins/Messaging/UdpMessaging/Source/UdpMessaging/Private/Tunnel/UdpMessageTunnel.cpp
+++ b/Engine/Plugins/Messaging/UdpMessaging/Source/UdpMessaging/Private/Tunnel/UdpMessageTunnel.cpp
@@ -83,6 +83,12 @@ FUdpMessageTunnel::~FUdpMessageTunnel()
 /* FRunnable interface
  *****************************************************************************/
 
+FSingleThreadRunnable* FUdpMessageTunnel::GetSingleThreadInterface()
+{
+	return this;
+}
+
+
 bool FUdpMessageTunnel::Init()
 {
 	return true;
@@ -93,17 +99,8 @@ uint32 FUdpMessageTunnel::Run()
 {
 	while (!Stopping)
 	{
-		CurrentTime = FDateTime::UtcNow();
-
-		UpdateConnections();
-
-		UdpToTcp(MulticastSocket);
-		UdpToTcp(UnicastSocket);
-
-		TcpToUdp();
-
-		RemoveExpiredNodes(LocalNodes);
-		RemoveExpiredNodes(RemoteNodes);
+		Update();
+		FPlatformProcess::Sleep(0.0f);
 	}
 
 	return 0;
@@ -333,6 +330,22 @@ void FUdpMessageTunnel::UdpToTcp(FSocket* Socket)
 }
 
 
+void FUdpMessageTunnel::Update()
+{
+	CurrentTime = FDateTime::UtcNow();
+
+	UpdateConnections();
+
+	UdpToTcp(MulticastSocket);
+	UdpToTcp(UnicastSocket);
+
+	TcpToUdp();
+
+	RemoveExpiredNodes(LocalNodes);
+	RemoveExpiredNodes(RemoteNodes);
+}
+
+
 void FUdpMessageTunnel::UpdateConnections()
 {
 	bool ConnectionsChanged = false;
@@ -368,6 +381,15 @@ void FUdpMessageTunnel::UpdateConnections()
 }
 
 
+/* FSingleThreadRunnable interface
+ *****************************************************************************/
+
+void FUdpMessageTunnel::Tick()
+{
+	Update();
+}
+
+
 /* FUdpMessageTunnel callbacks
  *****************************************************************************/
 
diff --git a/Engine/Plugins/Messaging/UdpMessaging/Source/UdpMessaging/Private/Tunnel/UdpMessageTunnel.h b/Engine/Plugins/Messaging/UdpMessaging/Source/UdpMessaging/Private/Tunnel/UdpMessageTunnel.h
--- a/Engine/Plugins/Messaging/UdpMessaging/Source/UdpMessaging/Private/Tunnel/UdpMessageTunnel.h
+++ b/Engine/Plugins/Messaging/UdpMessaging/Source/UdpMessaging/Private/Tunnel/UdpMessageTunnel.h
@@ -8,6 +8,7 @@
 #include "Containers/Map.h"
 #include "Containers/Queue.h"
 #include "Interfaces/IPv4/IPv4Endpoint.h"
+#include "Misc/SingleThreadRunnable.h"
 #include "Templates/SharedPointer.h"
 
 #include "IUdpMessageTunnel.h"
@@ -23,6 +24,7 @@ class FUdpMessageTunnelConnection;
 class FUdpMessageTunnel
 	: public FRunnable
 	, public IUdpMessageTunnel
+	, private FSingleThreadRunnable
 {
 	/** Structure for transport node information. */
 	struct FNodeInfo
@@ -54,6 +56,7 @@ public:
 
 	//~ FRunnable interface
 
+	virtual FSingleThreadRunnable* GetSingleThreadInterface() override;
 	virtual bool Init() override;
 	virtual uint32 Run() override;
 	virtual void Stop() override;
@@ -96,9 +99,18 @@ protected:
 	 */
 	void UdpToTcp(FSocket* Socket);
 
+	/** Update the tunnel. */
+	void Update();
+
 	/** Updates all active and pending connections. */
 	void UpdateConnections();
 
+protected:
+
+	//~ FSingleThreadRunnable interface
+
+	virtual void Tick() override;
+
 private:
 
 	/** Callback for accepted connections to the local tunnel server. */
diff --git a/Engine/Plugins/Messaging/UdpMessaging/Source/UdpMessaging/Private/UdpMessagingModule.cpp b/Engine/Plugins/Messaging/UdpMessaging/Source/UdpMessaging/Private/UdpMessagingModule.cpp
--- a/Engine/Plugins/Messaging/UdpMessaging/Source/UdpMessaging/Private/UdpMessagingModule.cpp
+++ b/Engine/Plugins/Messaging/UdpMessaging/Source/UdpMessaging/Private/UdpMessagingModule.cpp
@@ -442,12 +442,6 @@ protected:
 			return false;
 		}
 
-		// single thread support not implemented yet
-		if (!FPlatformProcess::SupportsMultithreading())
-		{
-			return false;
-		}
-
 		// always allow in standalone Slate applications
 		if (!FApp::IsGame() && !IsRunningCommandlet())
 		{
diff --git a/Engine/Plugins/MovieScene/ActorSequence/Source/ActorSequence/Private/ActorSequence.cpp b/Engine/Plugins/MovieScene/ActorSequence/Source/ActorSequence/Private/ActorSequence.cpp
--- a/Engine/Plugins/MovieScene/ActorSequence/Source/ActorSequence/Private/ActorSequence.cpp
+++ b/Engine/Plugins/MovieScene/ActorSequence/Source/ActorSequence/Private/ActorSequence.cpp
@@ -15,10 +15,28 @@ IMPLEMENT_MODULE(FDefaultModuleImpl, ActorSequence);
 UActorSequence::FOnInitialize UActorSequence::OnInitializeSequenceEvent;
 #endif
 
+static TAutoConsoleVariable<int32> CVarDefaultEvaluationType(
+	TEXT("ActorSequence.DefaultEvaluationType"),
+	0,
+	TEXT("0: Playback locked to playback frames\n1: Unlocked playback with sub frame interpolation"),
+	ECVF_Default);
+
+static TAutoConsoleVariable<FString> CVarDefaultFrameResolution(
+	TEXT("ActorSequence.DefaultFrameResolution"),
+	TEXT("24000fps"),
+	TEXT("Specifies default a frame resolution for newly created level sequences. Examples: 30 fps, 120/1 (120 fps), 30000/1001 (29.97), 0.01s (10ms)."),
+	ECVF_Default);
+
+static TAutoConsoleVariable<FString> CVarDefaultDisplayRate(
+	TEXT("ActorSequence.DefaultDisplayRate"),
+	TEXT("30fps"),
+	TEXT("Specifies default a display frame rate for newly created level sequences; also defines frame locked frame rate where sequences are set to be frame locked. Examples: 30 fps, 120/1 (120 fps), 30000/1001 (29.97), 0.01s (10ms)."),
+	ECVF_Default);
+
 UActorSequence::UActorSequence(const FObjectInitializer& ObjectInitializer)
 	: Super(ObjectInitializer)
 	, MovieScene(nullptr)
-#if WITH_EDITOR
+#if WITH_EDITORONLY_DATA
 	, bHasBeenInitialized(false)
 #endif
 {
@@ -52,14 +70,30 @@ UBlueprint* UActorSequence::GetParentBlueprint() const
 void UActorSequence::PostInitProperties()
 {
 #if WITH_EDITOR && WITH_EDITORONLY_DATA
+
+	// We do not run the default initialization for actor sequences that are CDOs, or that are going to be loaded (since they will have already been initialized in that case)
+	EObjectFlags ExcludeFlags = RF_ClassDefaultObject | RF_NeedLoad | RF_NeedPostLoad | RF_NeedPostLoadSubobjects | RF_WasLoaded;
+
 	UActorComponent* OwnerComponent = Cast<UActorComponent>(GetOuter());
-	if (!bHasBeenInitialized && !HasAnyFlags(RF_ClassDefaultObject) && OwnerComponent && !OwnerComponent->HasAnyFlags(RF_ClassDefaultObject))
+	if (!bHasBeenInitialized && !HasAnyFlags(ExcludeFlags) && OwnerComponent && !OwnerComponent->HasAnyFlags(ExcludeFlags))
 	{
 		AActor* Actor = Cast<AActor>(OwnerComponent->GetOuter());
 
 		FGuid BindingID = MovieScene->AddPossessable(Actor ? Actor->GetActorLabel() : TEXT("Owner"), Actor ? Actor->GetClass() : AActor::StaticClass());
 		ObjectReferences.CreateBinding(BindingID, FActorSequenceObjectReference::CreateForContextActor());
 
+		const bool bFrameLocked = CVarDefaultEvaluationType.GetValueOnGameThread() != 0;
+
+		MovieScene->SetEvaluationType( bFrameLocked ? EMovieSceneEvaluationType::FrameLocked : EMovieSceneEvaluationType::WithSubFrames );
+
+		FFrameRate FrameResolution(60000, 1);
+		TryParseString(FrameResolution, *CVarDefaultFrameResolution.GetValueOnGameThread());
+		MovieScene->SetFrameResolutionDirectly(FrameResolution);
+
+		FFrameRate PlaybackFrameRate(30, 1);
+		TryParseString(PlaybackFrameRate, *CVarDefaultDisplayRate.GetValueOnGameThread());
+		MovieScene->SetPlaybackFrameRate(PlaybackFrameRate);
+
 		OnInitializeSequenceEvent.Broadcast(this);
 		bHasBeenInitialized = true;
 	}
diff --git a/Engine/Plugins/MovieScene/ActorSequenceEditor/Source/ActorSequenceEditor/Private/ActorSequenceEditor.cpp b/Engine/Plugins/MovieScene/ActorSequenceEditor/Source/ActorSequenceEditor/Private/ActorSequenceEditor.cpp
--- a/Engine/Plugins/MovieScene/ActorSequenceEditor/Source/ActorSequenceEditor/Private/ActorSequenceEditor.cpp
+++ b/Engine/Plugins/MovieScene/ActorSequenceEditor/Source/ActorSequenceEditor/Private/ActorSequenceEditor.cpp
@@ -133,7 +133,12 @@ public:
 	static void OnInitializeSequence(UActorSequence* Sequence)
 	{
 		auto* ProjectSettings = GetDefault<UMovieSceneToolsProjectSettings>();
-		Sequence->GetMovieScene()->SetPlaybackRange(ProjectSettings->DefaultStartTime, ProjectSettings->DefaultStartTime + ProjectSettings->DefaultDuration);
+		UMovieScene* MovieScene = Sequence->GetMovieScene();
+		
+		FFrameNumber StartFrame = (ProjectSettings->DefaultStartTime * MovieScene->GetFrameResolution()).RoundToFrame();
+		int32        Duration   = (ProjectSettings->DefaultDuration * MovieScene->GetFrameResolution()).RoundToFrame().Value;
+
+		MovieScene->SetPlaybackRange(StartFrame, Duration);
 	}
 
 	/** Register details view customizations. */
diff --git a/Engine/Plugins/MovieScene/LevelSequenceEditor/Source/LevelSequenceEditor/Private/CinematicViewport/SCinematicLevelViewport.cpp b/Engine/Plugins/MovieScene/LevelSequenceEditor/Source/LevelSequenceEditor/Private/CinematicViewport/SCinematicLevelViewport.cpp
--- a/Engine/Plugins/MovieScene/LevelSequenceEditor/Source/LevelSequenceEditor/Private/CinematicViewport/SCinematicLevelViewport.cpp
+++ b/Engine/Plugins/MovieScene/LevelSequenceEditor/Source/LevelSequenceEditor/Private/CinematicViewport/SCinematicLevelViewport.cpp
@@ -21,7 +21,7 @@
 #include "Tracks/MovieSceneSubTrack.h"
 #include "Tracks/MovieSceneCinematicShotTrack.h"
 #include "Sections/MovieSceneCinematicShotSection.h"
-#include "ISequencerKeyCollection.h"
+#include "SequencerKeyCollection.h"
 #include "CinematicViewport/SCinematicTransportRange.h"
 #include "CinematicViewport/FilmOverlays.h"
 #include "Widgets/Layout/SWidgetSwitcher.h"
@@ -46,12 +46,12 @@ struct SNonThrottledSpinBox : SSpinBox<T>
 	}
 };
 
-struct FTypeInterfaceProxy : INumericTypeInterface<float>
+struct FTypeInterfaceProxy : INumericTypeInterface<double>
 {
-	TSharedPtr<INumericTypeInterface<float>> Impl;
+	TSharedPtr<INumericTypeInterface<double>> Impl;
 
 	/** Convert the type to/from a string */
-	virtual FString ToString(const float& Value) const override
+	virtual FString ToString(const double& Value) const override
 	{
 		if (Impl.IsValid())
 		{
@@ -60,13 +60,13 @@ struct FTypeInterfaceProxy : INumericTypeInterface<float>
 		return FString();
 	}
 
-	virtual TOptional<float> FromString(const FString& InString, const float& InExistingValue) override
+	virtual TOptional<double> FromString(const FString& InString, const double& InExistingValue) override
 	{
 		if (Impl.IsValid())
 		{
 			return Impl->FromString(InString, InExistingValue);
 		}
-		return TOptional<float>();
+		return TOptional<double>();
 	}
 
 	/** Check whether the typed character is valid */
@@ -195,16 +195,21 @@ void SCinematicLevelViewport::Construct(const FArguments& InArgs)
 			.BorderImage(nullptr)
 			.ForegroundColor(FEditorStyle::GetSlateColor("SelectionColor").GetColor(FWidgetStyle()))
 			[
-				SNew(SNonThrottledSpinBox<float>)
+				SNew(SNonThrottledSpinBox<double>)
 				.TypeInterface(TypeInterfaceProxy)
 				.Style(FEditorStyle::Get(), "Sequencer.HyperlinkSpinBox")
 				.Font(FEditorStyle::GetFontStyle("Sequencer.FixedFont"))
 				.OnValueCommitted(this, &SCinematicLevelViewport::OnTimeCommitted)
 				.OnValueChanged(this, &SCinematicLevelViewport::SetTime)
+				.MinValue(TOptional<double>())
+				.MaxValue(TOptional<double>())
 				.OnEndSliderMovement(this, &SCinematicLevelViewport::SetTime)
-				.MinValue(this, &SCinematicLevelViewport::GetMinTime)
-				.MaxValue(this, &SCinematicLevelViewport::GetMaxTime)
 				.Value(this, &SCinematicLevelViewport::GetTime)
+				.Delta_Lambda([=]()
+				{
+					return UIData.OuterResolution.AsDecimal() * UIData.OuterPlayRate.AsInterval(); 
+				})
+				.LinearDeltaSensitivity(25)
 			]
 		]
 
@@ -332,7 +337,7 @@ void SCinematicLevelViewport::Construct(const FArguments& InArgs)
 								SNew(STextBlock)
 								.Font(FEditorStyle::GetFontStyle("Sequencer.FixedFont"))
 								.ColorAndOpacity(Gray)
-								.Text_Lambda([=]{ return UIData.Frame; })
+								.Text_Lambda([=]{ return UIData.LocalPlaybackTime; })
 							]
 						]
 					]
@@ -431,48 +436,60 @@ EVisibility SCinematicLevelViewport::GetControlsVisibility() const
 	return CurrentToolkit.IsValid() ? EVisibility::Visible : EVisibility::Collapsed;
 }
 
-TOptional<float> SCinematicLevelViewport::GetMinTime() const
+TOptional<double> SCinematicLevelViewport::GetMinTime() const
 {
 	ISequencer* Sequencer = GetSequencer();
 	if (Sequencer)
 	{
-		return Sequencer->GetFocusedMovieSceneSequence()->GetMovieScene()->GetEditorData().WorkingRange.GetLowerBoundValue();
+		FFrameRate   PlayRate      = Sequencer->GetLocalTime().Rate;
+		UMovieScene* MovieScene    = Sequencer->GetFocusedMovieSceneSequence()->GetMovieScene();
+		double       TimeInSeconds = MovieScene->GetEditorData().WorkStart;
+
+		return (TimeInSeconds*PlayRate).GetFrame().Value;
 	}
-	return TOptional<float>();
+	return TOptional<double>();
 }
 
-TOptional<float> SCinematicLevelViewport::GetMaxTime() const
+TOptional<double> SCinematicLevelViewport::GetMaxTime() const
 {
 	ISequencer* Sequencer = GetSequencer();
 	if (Sequencer)
 	{
-		return Sequencer->GetFocusedMovieSceneSequence()->GetMovieScene()->GetEditorData().WorkingRange.GetUpperBoundValue();
+		FFrameRate   PlayRate      = Sequencer->GetLocalTime().Rate;
+		UMovieScene* MovieScene    = Sequencer->GetFocusedMovieSceneSequence()->GetMovieScene();
+		double       TimeInSeconds = MovieScene->GetEditorData().WorkEnd;
+
+		return (TimeInSeconds*PlayRate).GetFrame().Value;
 	}
-	return TOptional<float>();
+	return TOptional<double>();
 }
 
-void SCinematicLevelViewport::OnTimeCommitted(float Value, ETextCommit::Type)
+void SCinematicLevelViewport::OnTimeCommitted(double Value, ETextCommit::Type)
 {
 	SetTime(Value);
 }
 
-void SCinematicLevelViewport::SetTime(float Value)
+void SCinematicLevelViewport::SetTime(double Value)
 {
+	// Clamp the value as the UI can't due to needing an unbounded spinbox for value-change-rate purposes.
+	Value = FMath::Clamp(Value, GetMinTime().GetValue(), GetMaxTime().GetValue());
+
 	ISequencer* Sequencer = GetSequencer();
 	if (Sequencer)
 	{
-		Sequencer->SetLocalTime(Value);
+		FFrameRate SequencerPlayRate = Sequencer->GetLocalTime().Rate;
+		Sequencer->SetLocalTime(FFrameTime::FromDecimal(Value));
 	}
 }
 
-float SCinematicLevelViewport::GetTime() const
+double SCinematicLevelViewport::GetTime() const
 {
 	ISequencer* Sequencer = GetSequencer();
 	if (Sequencer)
 	{
-		return Sequencer->GetLocalTime();
+		return Sequencer->GetLocalTime().Time.GetFrame().Value;
 	}
-	return 0.f;
+	return 0;
 }
 
 void SCinematicLevelViewport::CacheDesiredViewportSize(const FGeometry& AllottedGeometry)
@@ -526,7 +543,7 @@ void SCinematicLevelViewport::Setup(FLevelSequenceEditorToolkit& NewToolkit)
 	ISequencer* Sequencer = GetSequencer();
 	if (Sequencer)
 	{
-		TypeInterfaceProxy->Impl = Sequencer->GetZeroPadNumericTypeInterface();
+		TypeInterfaceProxy->Impl = Sequencer->GetNumericTypeInterface();
 
 		if (TransportRange.IsValid())
 		{
@@ -611,41 +628,40 @@ void SCinematicLevelViewport::Tick(const FGeometry& AllottedGeometry, const doub
 		SubTrack = Cast<UMovieSceneSubTrack>(Sequence->GetMovieScene()->FindMasterTrack(UMovieSceneSubTrack::StaticClass()));
 	}
 
+	const FFrameRate OuterResolution = Sequencer->GetFocusedFrameResolution();
+	const FFrameRate OuterPlayRate   = Sequencer->GetFocusedPlayRate();
+
+	const FFrameTime OuterTime       = Sequencer->GetLocalTime().ConvertTo(OuterResolution);
+	UIData.OuterResolution = OuterResolution;
+	UIData.OuterPlayRate = OuterPlayRate;
+
+
 	UMovieSceneSubSection* SubSection = nullptr;
 	if (SubTrack)
 	{
-		const float CurrentTime = Sequencer->GetLocalTime();
-
 		for (UMovieSceneSection* Section : SubTrack->GetAllSections())
 		{
-			if (Section->IsInfinite() || Section->IsTimeWithinSection(CurrentTime))
+			if (Section->GetRange().Contains(OuterTime.FrameNumber))
 			{
 				SubSection = CastChecked<UMovieSceneSubSection>(Section);
 			}
 		}
 	}
 
-	const float AbsoluteTime = Sequencer->GetLocalTime();
-
 	FText TimeFormat = LOCTEXT("TimeFormat", "{0}");
 
-	TSharedPtr<INumericTypeInterface<float>> ZeroPadTypeInterface = Sequencer->GetZeroPadNumericTypeInterface();
+	TSharedPtr<INumericTypeInterface<double>> TimeDisplayFormatInterface = Sequencer->GetNumericTypeInterface();
 
-	if (SubSection)
+	UMovieSceneSequence* SubSequence = SubSection ? SubSection->GetSequence() : nullptr;
+	if (SubSequence)
 	{
-		float PlaybackRangeStart = 0.f;
-		if (SubSection->GetSequence() != nullptr)
-		{
-			UMovieScene* ShotMovieScene = SubSection->GetSequence()->GetMovieScene();
-			PlaybackRangeStart = ShotMovieScene->GetPlaybackRange().GetLowerBoundValue();
-		}
-
-		const float InnerOffset = (AbsoluteTime - SubSection->GetStartTime()) * SubSection->Parameters.TimeScale;
-		const float AbsoluteShotPosition = PlaybackRangeStart + SubSection->Parameters.StartOffset + InnerOffset;
+		FFrameRate                   InnerResolution       = SubSequence->GetMovieScene()->GetFrameResolution();
+		FMovieSceneSequenceTransform OuterToInnerTransform = SubSection->OuterToInnerTransform();
+		const FFrameTime             InnerShotPosition	   = OuterTime * OuterToInnerTransform;
 
-		UIData.Frame = FText::Format(
+		UIData.LocalPlaybackTime = FText::Format(
 			TimeFormat,
-			FText::FromString(ZeroPadTypeInterface->ToString(AbsoluteShotPosition))
+			FText::FromString(TimeDisplayFormatInterface->ToString(InnerShotPosition.GetFrame().Value))
 		);
 
 		UMovieSceneCinematicShotSection* CinematicShotSection = Cast<UMovieSceneCinematicShotSection>(SubSection);
@@ -660,24 +676,28 @@ void SCinematicLevelViewport::Tick(const FGeometry& AllottedGeometry, const doub
 	}
 	else
 	{
-		UIData.Frame = FText::Format(
+		const FFrameTime DisplayTime = Sequencer->GetLocalTime().Time;
+
+		UIData.LocalPlaybackTime = FText::Format(
 			TimeFormat,
-			FText::FromString(ZeroPadTypeInterface->ToString(AbsoluteTime))
+			FText::FromString(TimeDisplayFormatInterface->ToString(DisplayTime.GetFrame().Value))
 			);
 
 		UIData.ShotName = Sequence->GetDisplayName();
 	}
 
-	TRange<float> EntireRange = Sequence->GetMovieScene()->GetEditorData().WorkingRange;
+	const FMovieSceneEditorData& EditorData = Sequence->GetMovieScene()->GetEditorData();
 
+	FQualifiedFrameTime MasterStartTime(EditorData.WorkStart * OuterPlayRate, OuterPlayRate);
 	UIData.MasterStartText = FText::Format(
 		TimeFormat,
-		FText::FromString(ZeroPadTypeInterface->ToString(EntireRange.GetLowerBoundValue()))
+		FText::FromString(TimeDisplayFormatInterface->ToString(MasterStartTime.Time.GetFrame().Value))
 	);
 
+	FQualifiedFrameTime MasterEndTime(EditorData.WorkEnd * OuterPlayRate, OuterPlayRate);
 	UIData.MasterEndText = FText::Format(
 		TimeFormat,
-		FText::FromString(ZeroPadTypeInterface->ToString(EntireRange.GetUpperBoundValue()))
+		FText::FromString(TimeDisplayFormatInterface->ToString(MasterEndTime.Time.GetFrame().Value))
 	);
 
 	UIData.CameraName = FText::GetEmpty();
diff --git a/Engine/Plugins/MovieScene/LevelSequenceEditor/Source/LevelSequenceEditor/Private/CinematicViewport/SCinematicLevelViewport.h b/Engine/Plugins/MovieScene/LevelSequenceEditor/Source/LevelSequenceEditor/Private/CinematicViewport/SCinematicLevelViewport.h
--- a/Engine/Plugins/MovieScene/LevelSequenceEditor/Source/LevelSequenceEditor/Private/CinematicViewport/SCinematicLevelViewport.h
+++ b/Engine/Plugins/MovieScene/LevelSequenceEditor/Source/LevelSequenceEditor/Private/CinematicViewport/SCinematicLevelViewport.h
@@ -9,6 +9,7 @@
 #include "Widgets/DeclarativeSyntaxSupport.h"
 #include "Widgets/SCompoundWidget.h"
 #include "LevelEditorViewport.h"
+#include "FrameRate.h"
 
 class FLevelSequenceEditorToolkit;
 class FLevelViewportLayout;
@@ -20,6 +21,8 @@ class SCinematicPreviewViewport;
 class SCinematicTransportRange;
 class SEditorViewport;
 class SLevelViewport;
+
+struct FQualifiedFrameTime;
 struct FTypeInterfaceProxy;
 
 /** Overridden level viewport client for this viewport */
@@ -45,12 +48,17 @@ struct FUIData
 	FText Lens;
 	/** The name of the current shot's filmback */
 	FText Filmback;
-	/** The text that represents the current frame */
-	FText Frame;
+	/** The text that represents the current playback time relative to the currently evaluated sequence. */
+	FText LocalPlaybackTime;
 	/** The text that represents the master start frame */
 	FText MasterStartText;
 	/** The text that represents the master end frame */
 	FText MasterEndText;
+
+	/** The frame resolution of the master */
+	FFrameRate OuterResolution;
+	/** The play rate of the master */
+	FFrameRate OuterPlayRate;
 };
 
 
@@ -114,15 +122,15 @@ private:
 
 private:
 
-	TOptional<float> GetMinTime() const;
-
-	TOptional<float> GetMaxTime() const;
+	TOptional<double> GetMinTime() const;
+			  
+	TOptional<double> GetMaxTime() const;
 
-	void OnTimeCommitted(float Value, ETextCommit::Type);
+	void OnTimeCommitted(double Value, ETextCommit::Type);
 
-	void SetTime(float Value);
+	void SetTime(double Value);
 
-	float GetTime() const;
+	double GetTime() const;
 
 private:
 	
diff --git a/Engine/Plugins/MovieScene/LevelSequenceEditor/Source/LevelSequenceEditor/Private/CinematicViewport/SCinematicTransportRange.cpp b/Engine/Plugins/MovieScene/LevelSequenceEditor/Source/LevelSequenceEditor/Private/CinematicViewport/SCinematicTransportRange.cpp
--- a/Engine/Plugins/MovieScene/LevelSequenceEditor/Source/LevelSequenceEditor/Private/CinematicViewport/SCinematicTransportRange.cpp
+++ b/Engine/Plugins/MovieScene/LevelSequenceEditor/Source/LevelSequenceEditor/Private/CinematicViewport/SCinematicTransportRange.cpp
@@ -2,7 +2,7 @@
 
 #include "CinematicViewport/SCinematicTransportRange.h"
 #include "Rendering/DrawElements.h"
-#include "ISequencerKeyCollection.h"
+#include "SequencerKeyCollection.h"
 #include "MovieSceneSequence.h"
 #include "EditorStyleSet.h"
 #include "Styles/LevelSequenceEditorStyle.h"
@@ -12,7 +12,6 @@
 
 #define LOCTEXT_NAMESPACE "SCinematicTransportRange"
 
-
 void SCinematicTransportRange::Construct(const FArguments& InArgs)
 {
 	bDraggingTime = false;
@@ -79,10 +78,11 @@ void SCinematicTransportRange::SetTime(const FGeometry& MyGeometry, const FPoint
 	{
 		float Lerp = MyGeometry.AbsoluteToLocal(MouseEvent.GetScreenSpacePosition()).X / MyGeometry.GetLocalSize().X;
 		Lerp = FMath::Clamp(Lerp, 0.f, 1.f);
-		
-		const TRange<float> WorkingRange = Sequencer->GetFocusedMovieSceneSequence()->GetMovieScene()->GetEditorData().WorkingRange;
-		
-		Sequencer->SetLocalTime(WorkingRange.GetLowerBoundValue() + WorkingRange.Size<float>()*Lerp, ESnapTimeMode::STM_All);
+
+		FMovieSceneEditorData& EditorData = Sequencer->GetFocusedMovieSceneSequence()->GetMovieScene()->GetEditorData();
+		double NewTimeSeconds = EditorData.WorkStart + (EditorData.WorkEnd - EditorData.WorkStart) * Lerp;
+
+		Sequencer->SetLocalTime(NewTimeSeconds * Sequencer->GetFocusedFrameResolution(), ESnapTimeMode::STM_All);
 	}
 }
 
@@ -96,11 +96,11 @@ void SCinematicTransportRange::Tick(const FGeometry& AllottedGeometry, const dou
 	ISequencer* Sequencer = GetSequencer();
 	if (Sequencer)
 	{
-		TRange<float> WorkingRange = Sequencer->GetFocusedMovieSceneSequence()->GetMovieScene()->GetEditorData().WorkingRange;
+		TRange<double> WorkingRange = Sequencer->GetFocusedMovieSceneSequence()->GetMovieScene()->GetEditorData().GetWorkingRange();
 
 		// Anything within 3 pixel's worth of time is a duplicate as far as we're concerned
-		FTimeToPixel TimeToPixelConverter(AllottedGeometry, WorkingRange);
-		const float DuplicateThreshold = (TimeToPixelConverter.PixelToTime(3.f) - TimeToPixelConverter.PixelToTime(0.f));
+		FTimeToPixel TimeToPixelConverter(AllottedGeometry, WorkingRange, Sequencer->GetFocusedFrameResolution());
+		const float DuplicateThreshold = (TimeToPixelConverter.PixelToSeconds(3.f) - TimeToPixelConverter.PixelToSeconds(0.f));
 
 		Sequencer->GetKeysFromSelection(ActiveKeyCollection, DuplicateThreshold);
 	}
@@ -119,10 +119,11 @@ int32 SCinematicTransportRange::OnPaint(const FPaintArgs& Args, const FGeometry&
 	static const float TrackOffsetY = 6.f;
 	const float TrackHeight = AllottedGeometry.GetLocalSize().Y - TrackOffsetY;
 
-	TRange<float> WorkingRange = Sequencer->GetFocusedMovieSceneSequence()->GetMovieScene()->GetEditorData().WorkingRange;
-	TRange<float> PlaybackRange = Sequencer->GetFocusedMovieSceneSequence()->GetMovieScene()->GetPlaybackRange();
+	FFrameRate           FrameResolution = Sequencer->GetFocusedFrameResolution();
+	TRange<double>       WorkingRange    = Sequencer->GetFocusedMovieSceneSequence()->GetMovieScene()->GetEditorData().GetWorkingRange();
+	TRange<FFrameNumber> PlaybackRange   = Sequencer->GetFocusedMovieSceneSequence()->GetMovieScene()->GetPlaybackRange();
 
-	const float TimePerPixel = WorkingRange.Size<float>() / AllottedGeometry.GetLocalSize().X;
+	const FFrameTime FramesPerPixel = (WorkingRange.Size<double>() / AllottedGeometry.GetLocalSize().X) * FrameResolution;
 
 	FColor DarkGray(40, 40, 40);
 	FColor MidGray(80, 80, 80);
@@ -140,8 +141,8 @@ int32 SCinematicTransportRange::OnPaint(const FPaintArgs& Args, const FGeometry&
 	
 	const float FullRange = WorkingRange.Size<float>();
 
-	const float PlaybackStartLerp	= (PlaybackRange.GetLowerBoundValue() - WorkingRange.GetLowerBoundValue()) / FullRange;
-	const float PlaybackEndLerp		= (PlaybackRange.GetUpperBoundValue() - WorkingRange.GetLowerBoundValue()) / FullRange;
+	const float PlaybackStartLerp	= (PlaybackRange.GetLowerBoundValue()/FrameResolution - WorkingRange.GetLowerBoundValue()) / FullRange;
+	const float PlaybackEndLerp		= (PlaybackRange.GetUpperBoundValue()/FrameResolution - WorkingRange.GetLowerBoundValue()) / FullRange;
 
 	// Draw the playback range
 	FSlateDrawElement::MakeBox(
@@ -153,8 +154,8 @@ int32 SCinematicTransportRange::OnPaint(const FPaintArgs& Args, const FGeometry&
 		FLinearColor(MidGray)
 	);
 
-	const float CurrentTime = Sequencer->GetLocalTime();
-	const float ProgressLerp = (CurrentTime - WorkingRange.GetLowerBoundValue()) / FullRange;
+	const FQualifiedFrameTime CurrentTime = Sequencer->GetLocalTime();
+	const float ProgressLerp = (CurrentTime.AsSeconds() - WorkingRange.GetLowerBoundValue()) / FullRange;
 
 	// Draw the playback progress
 	if (ProgressLerp > PlaybackStartLerp)
@@ -184,19 +185,19 @@ int32 SCinematicTransportRange::OnPaint(const FPaintArgs& Args, const FGeometry&
 			const float BrushOffsetY = TrackOffsetY + TrackHeight * .5f - BrushHeight * .5f;
 			const FSlateBrush* KeyBrush = FLevelSequenceEditorStyle::Get()->GetBrush("LevelSequenceEditor.CinematicViewportTransportRangeKey");
 
-			ActiveKeyCollection->IterateKeys([&](float Time){
-				if (Time < WorkingRange.GetLowerBoundValue() || Time > WorkingRange.GetUpperBoundValue())
-				{
-					// continue iteration
-					return true;
-				}
+			TRange<FFrameNumber> VisibleFrameRange(
+				(WorkingRange.GetLowerBoundValue() * FrameResolution).FloorToFrame(), 
+				(WorkingRange.GetUpperBoundValue() * FrameResolution).CeilToFrame()
+				);
 
-				if (FMath::IsNearlyEqual(CurrentTime, Time, TimePerPixel*.5f))
+			for (FFrameNumber Time : ActiveKeyCollection->GetKeysInRange(VisibleFrameRange))
+			{
+				if (FMath::Abs(CurrentTime.Time - Time) < FramesPerPixel/2)
 				{
 					bPlayMarkerOnKey = true;
 				}
 
-				float Lerp = (Time - WorkingRange.GetLowerBoundValue()) / FullRange;
+				float Lerp = (Time/FrameResolution - WorkingRange.GetLowerBoundValue()) / FullRange;
 
 				FSlateDrawElement::MakeBox(
 					OutDrawElements,
@@ -209,10 +210,7 @@ int32 SCinematicTransportRange::OnPaint(const FPaintArgs& Args, const FGeometry&
 					DrawEffects,
 					KeyFrameColor
 				);
-
-				// continue iteration
-				return true;
-			});
+			}
 		}
 	}
 
diff --git a/Engine/Plugins/MovieScene/LevelSequenceEditor/Source/LevelSequenceEditor/Private/CinematicViewport/SCinematicTransportRange.h b/Engine/Plugins/MovieScene/LevelSequenceEditor/Source/LevelSequenceEditor/Private/CinematicViewport/SCinematicTransportRange.h
--- a/Engine/Plugins/MovieScene/LevelSequenceEditor/Source/LevelSequenceEditor/Private/CinematicViewport/SCinematicTransportRange.h
+++ b/Engine/Plugins/MovieScene/LevelSequenceEditor/Source/LevelSequenceEditor/Private/CinematicViewport/SCinematicTransportRange.h
@@ -10,7 +10,7 @@
 class FPaintArgs;
 class FSlateWindowElementList;
 class ISequencer;
-class ISequencerKeyCollection;
+class FSequencerKeyCollection;
 
 class SCinematicTransportRange : public SCompoundWidget
 {
@@ -43,7 +43,7 @@ private:
 	TWeakPtr<ISequencer> WeakSequencer;
 
 	/** The collection of keys for the currently active sequencer selection */
-	TUniquePtr<ISequencerKeyCollection> ActiveKeyCollection;
+	TUniquePtr<FSequencerKeyCollection> ActiveKeyCollection;
 
 	bool bDraggingTime;
 };
diff --git a/Engine/Plugins/MovieScene/LevelSequenceEditor/Source/LevelSequenceEditor/Private/Factories/LevelSequenceFactoryNew.cpp b/Engine/Plugins/MovieScene/LevelSequenceEditor/Source/LevelSequenceEditor/Private/Factories/LevelSequenceFactoryNew.cpp
--- a/Engine/Plugins/MovieScene/LevelSequenceEditor/Source/LevelSequenceEditor/Private/Factories/LevelSequenceFactoryNew.cpp
+++ b/Engine/Plugins/MovieScene/LevelSequenceEditor/Source/LevelSequenceEditor/Private/Factories/LevelSequenceFactoryNew.cpp
@@ -27,10 +27,12 @@ UObject* ULevelSequenceFactoryNew::FactoryCreateNew(UClass* Class, UObject* InPa
 {
 	auto NewLevelSequence = NewObject<ULevelSequence>(InParent, Name, Flags|RF_Transactional);
 	NewLevelSequence->Initialize();
-	
+
 	// Set up some sensible defaults
 	const UMovieSceneToolsProjectSettings* ProjectSettings = GetDefault<UMovieSceneToolsProjectSettings>();
-	NewLevelSequence->GetMovieScene()->SetPlaybackRange(ProjectSettings->DefaultStartTime, ProjectSettings->DefaultStartTime + ProjectSettings->DefaultDuration);
+
+	FFrameRate FrameResolution = NewLevelSequence->GetMovieScene()->GetFrameResolution();
+	NewLevelSequence->GetMovieScene()->SetPlaybackRange((ProjectSettings->DefaultStartTime*FrameResolution).FloorToFrame(), (ProjectSettings->DefaultDuration*FrameResolution).FloorToFrame().Value);
 
 	return NewLevelSequence;
 }
diff --git a/Engine/Plugins/MovieScene/LevelSequenceEditor/Source/LevelSequenceEditor/Private/LevelSequenceEditorToolkit.cpp b/Engine/Plugins/MovieScene/LevelSequenceEditor/Source/LevelSequenceEditor/Private/LevelSequenceEditorToolkit.cpp
--- a/Engine/Plugins/MovieScene/LevelSequenceEditor/Source/LevelSequenceEditor/Private/LevelSequenceEditorToolkit.cpp
+++ b/Engine/Plugins/MovieScene/LevelSequenceEditor/Source/LevelSequenceEditor/Private/LevelSequenceEditorToolkit.cpp
@@ -8,6 +8,7 @@
 #include "AssetData.h"
 #include "Components/PrimitiveComponent.h"
 #include "Editor.h"
+#include "Channels/MovieSceneChannelProxy.h"
 #include "Containers/ArrayBuilder.h"
 #include "Modules/ModuleManager.h"
 #include "KeyParams.h"
@@ -449,20 +450,24 @@ void FLevelSequenceEditorToolkit::AddDefaultTracksForActor(AActor& Actor, const
 							Scale = ActorRelativeTransform.GetScale3D();
 						}
 
-						TransformSection->SetDefault(FTransformKey(EKey3DTransformChannel::Translation, EAxis::X, Location.X, false /*bUnwindRotation*/));
-						TransformSection->SetDefault(FTransformKey(EKey3DTransformChannel::Translation, EAxis::Y, Location.Y, false /*bUnwindRotation*/));
-						TransformSection->SetDefault(FTransformKey(EKey3DTransformChannel::Translation, EAxis::Z, Location.Z, false /*bUnwindRotation*/));
+						TArrayView<FMovieSceneFloatChannel*> FloatChannels = TransformSection->GetChannelProxy().GetChannels<FMovieSceneFloatChannel>();
+						FloatChannels[0]->SetDefault(Location.X);
+						FloatChannels[1]->SetDefault(Location.Y);
+						FloatChannels[2]->SetDefault(Location.Z);
 
-						TransformSection->SetDefault(FTransformKey(EKey3DTransformChannel::Rotation, EAxis::X, Rotation.Euler().X, false /*bUnwindRotation*/));
-						TransformSection->SetDefault(FTransformKey(EKey3DTransformChannel::Rotation, EAxis::Y, Rotation.Euler().Y, false /*bUnwindRotation*/));
-						TransformSection->SetDefault(FTransformKey(EKey3DTransformChannel::Rotation, EAxis::Z, Rotation.Euler().Z, false /*bUnwindRotation*/));
+						FloatChannels[3]->SetDefault(Rotation.Euler().X);
+						FloatChannels[4]->SetDefault(Rotation.Euler().Y);
+						FloatChannels[5]->SetDefault(Rotation.Euler().Z);
 
-						TransformSection->SetDefault(FTransformKey(EKey3DTransformChannel::Scale, EAxis::X, Scale.X, false /*bUnwindRotation*/));
-						TransformSection->SetDefault(FTransformKey(EKey3DTransformChannel::Scale, EAxis::Y, Scale.Y, false /*bUnwindRotation*/));
-						TransformSection->SetDefault(FTransformKey(EKey3DTransformChannel::Scale, EAxis::Z, Scale.Z, false /*bUnwindRotation*/));
+						FloatChannels[6]->SetDefault(Scale.X);
+						FloatChannels[7]->SetDefault(Scale.Y);
+						FloatChannels[8]->SetDefault(Scale.Z);
 					}
 
-					NewSection->SetIsInfinite(GetSequencer()->GetInfiniteKeyAreas());
+					if (GetSequencer()->GetInfiniteKeyAreas())
+					{
+						NewSection->SetRange(TRange<FFrameNumber>::All());
+					}
 				}
 			}
 		}
@@ -644,12 +649,12 @@ void FLevelSequenceEditorToolkit::HandleMapChanged(class UWorld* NewWorld, EMapC
 	}
 }
 
-void FLevelSequenceEditorToolkit::AddShot(UMovieSceneCinematicShotTrack* ShotTrack, const FString& ShotAssetName, const FString& ShotPackagePath, float ShotStartTime, float ShotEndTime, UObject* AssetToDuplicate, const FString& FirstShotAssetName)
+void FLevelSequenceEditorToolkit::AddShot(UMovieSceneCinematicShotTrack* ShotTrack, const FString& ShotAssetName, const FString& ShotPackagePath, FFrameNumber ShotStartTime, FFrameNumber ShotEndTime, UObject* AssetToDuplicate, const FString& FirstShotAssetName)
 {
 	// Create a level sequence asset for the shot
 	UObject* ShotAsset = LevelSequenceEditorHelpers::CreateLevelSequenceAsset(ShotAssetName, ShotPackagePath, AssetToDuplicate);
 	UMovieSceneSequence* ShotSequence = Cast<UMovieSceneSequence>(ShotAsset);
-	UMovieSceneSubSection* ShotSubSection = ShotTrack->AddSequence(ShotSequence, ShotStartTime, ShotEndTime-ShotStartTime);
+	UMovieSceneSubSection* ShotSubSection = ShotTrack->AddSequence(ShotSequence, ShotStartTime, (ShotEndTime-ShotStartTime).Value);
 
 	// Focus on the new shot
 	GetSequencer()->ForceEvaluate();
@@ -700,9 +705,8 @@ void FLevelSequenceEditorToolkit::AddShot(UMovieSceneCinematicShotTrack* ShotTra
 
 			if (SubSequence != nullptr)
 			{
-				UMovieSceneSubSection* SubSection = SubTrack->AddSequence(SubSequence, 0.f, ShotEndTime-ShotStartTime);
+				UMovieSceneSubSection* SubSection = SubTrack->AddSequence(SubSequence, 0, (ShotEndTime-ShotStartTime).Value);
 				SubSection->SetRowIndex(RowIndex++);
-				SubSection->SetStartTime(0.f);
 			}
 		}
 	}
@@ -741,8 +745,8 @@ void FLevelSequenceEditorToolkit::AddShot(UMovieSceneCinematicShotTrack* ShotTra
 		// Create a new camera cut section and add it to the camera cut track
 		CameraCutTrack = ShotSequence->GetMovieScene()->AddCameraCutTrack(UMovieSceneCameraCutTrack::StaticClass());
 		UMovieSceneCameraCutSection* CameraCutSection = NewObject<UMovieSceneCameraCutSection>(CameraCutTrack, NAME_None, RF_Transactional);
-		CameraCutSection->SetStartTime(ShotSequence->GetMovieScene()->GetPlaybackRange().GetLowerBoundValue()); 
-		CameraCutSection->SetEndTime(ShotSequence->GetMovieScene()->GetPlaybackRange().GetUpperBoundValue());
+
+		CameraCutSection->SetRange(ShotSequence->GetMovieScene()->GetPlaybackRange());
 		CameraCutSection->SetCameraGuid(CameraGuid);
 		CameraCutTrack->AddSection(*CameraCutSection);
 	}
@@ -760,15 +764,18 @@ void FLevelSequenceEditorToolkit::HandleMasterSequenceCreated(UObject* MasterSeq
 
 	UMovieSceneSequence* MasterSequence = Cast<UMovieSceneSequence>(MasterSequenceAsset);
 	UMovieSceneCinematicShotTrack* ShotTrack = MasterSequence->GetMovieScene()->AddMasterTrack<UMovieSceneCinematicShotTrack>();
-		
+
+	FFrameRate FrameResolution = MasterSequence->GetMovieScene()->GetFrameResolution();
+
 	// Create shots with a camera cut and a camera for each
-	float SequenceStartTime = ProjectSettings->DefaultStartTime;
-	float ShotStartTime = SequenceStartTime;
-	float ShotEndTime = ShotStartTime;
+	FFrameNumber SequenceStartTime = (ProjectSettings->DefaultStartTime * FrameResolution).FloorToFrame();
+	FFrameNumber ShotStartTime = SequenceStartTime;
+	FFrameNumber ShotEndTime   = ShotStartTime;
+	int32        ShotDuration  = (ProjectSettings->DefaultDuration * FrameResolution).RoundToFrame().Value;
 	FString FirstShotName; 
 	for (uint32 ShotIndex = 0; ShotIndex < NumShots; ++ShotIndex)
 	{
-		ShotEndTime += ProjectSettings->DefaultDuration;
+		ShotEndTime += ShotDuration;
 
 		FString ShotName = MovieSceneToolHelpers::GenerateNewShotName(ShotTrack->GetAllSections(), ShotStartTime);
 		FString ShotPackagePath = MovieSceneToolHelpers::GenerateNewShotPath(MasterSequence->GetMovieScene(), ShotName);
@@ -784,17 +791,18 @@ void FLevelSequenceEditorToolkit::HandleMasterSequenceCreated(UObject* MasterSeq
 		ShotStartTime = ShotEndTime;
 	}
 
-	MasterSequence->GetMovieScene()->SetPlaybackRange(SequenceStartTime, ShotEndTime);
+	MasterSequence->GetMovieScene()->SetPlaybackRange(SequenceStartTime, (ShotEndTime - SequenceStartTime).Value);
 
 #if WITH_EDITORONLY_DATA
-	const float OutputViewSize = ShotEndTime - SequenceStartTime;
-	const float OutputChange = OutputViewSize * 0.1f;
-	struct FMovieSceneEditorData EditorData;
-	EditorData.ViewRange = FFloatRange(SequenceStartTime - OutputChange, ShotEndTime + OutputChange);
-	EditorData.WorkingRange = FFloatRange(SequenceStartTime - OutputChange, ShotEndTime + OutputChange);
-	MasterSequence->GetMovieScene()->SetEditorData(EditorData);
+	const double SequenceStartSeconds = SequenceStartTime / FrameResolution;
+	const double SequenceEndSeconds   = ShotEndTime / FrameResolution;
+	const double OutputChange = (SequenceEndSeconds - SequenceStartSeconds) * 0.1;
+
+	FMovieSceneEditorData& EditorData = MasterSequence->GetMovieScene()->GetEditorData();
+	EditorData.ViewStart = EditorData.WorkStart = SequenceStartSeconds - OutputChange;
+	EditorData.ViewEnd   = EditorData.WorkEnd   = SequenceEndSeconds + OutputChange;
 #endif
-		
+
 	GetSequencer()->ResetToNewRootSequence(*MasterSequence);
 
 	UActorFactory* ActorFactory = GEditor->FindActorFactoryForActorClass(ALevelSequenceActor::StaticClass());
diff --git a/Engine/Plugins/MovieScene/LevelSequenceEditor/Source/LevelSequenceEditor/Private/LevelSequenceEditorToolkit.h b/Engine/Plugins/MovieScene/LevelSequenceEditor/Source/LevelSequenceEditor/Private/LevelSequenceEditorToolkit.h
--- a/Engine/Plugins/MovieScene/LevelSequenceEditor/Source/LevelSequenceEditor/Private/LevelSequenceEditorToolkit.h
+++ b/Engine/Plugins/MovieScene/LevelSequenceEditor/Source/LevelSequenceEditor/Private/LevelSequenceEditorToolkit.h
@@ -11,6 +11,8 @@
 #include "Framework/Docking/TabManager.h"
 #include "Toolkits/AssetEditorToolkit.h"
 
+struct FFrameNumber;
+
 class AActor;
 class FMenuBuilder;
 class ILevelViewport;
@@ -104,7 +106,7 @@ protected:
 	void AddDefaultTracksForActor(AActor& Actor, const FGuid Binding);
 	
 	/** Add a shot to a master sequence */
-	void AddShot(UMovieSceneCinematicShotTrack* ShotTrack, const FString& ShotAssetName, const FString& ShotPackagePath, float ShotStartTime, float ShotEndTime, UObject* AssetToDuplicate, const FString& FirstShotAssetName);
+	void AddShot(UMovieSceneCinematicShotTrack* ShotTrack, const FString& ShotAssetName, const FString& ShotPackagePath, FFrameNumber ShotStartTime, FFrameNumber ShotEndTime, UObject* AssetToDuplicate, const FString& FirstShotAssetName);
 
 	/** Called whenever sequencer has received focus */
 	void OnSequencerReceivedFocus();
diff --git a/Engine/Plugins/MovieScene/LevelSequenceEditor/Source/LevelSequenceEditor/Private/Misc/LevelSequenceEditorActorSpawner.cpp b/Engine/Plugins/MovieScene/LevelSequenceEditor/Source/LevelSequenceEditor/Private/Misc/LevelSequenceEditorActorSpawner.cpp
--- a/Engine/Plugins/MovieScene/LevelSequenceEditor/Source/LevelSequenceEditor/Private/Misc/LevelSequenceEditorActorSpawner.cpp
+++ b/Engine/Plugins/MovieScene/LevelSequenceEditor/Source/LevelSequenceEditor/Private/Misc/LevelSequenceEditorActorSpawner.cpp
@@ -16,6 +16,7 @@
 #include "Engine/World.h"
 #include "Settings/LevelEditorViewportSettings.h"
 #include "MovieSceneSequence.h"
+#include "Channels/MovieSceneChannelProxy.h"
 
 #define LOCTEXT_NAMESPACE "LevelSequenceEditorActorSpawner"
 
@@ -47,6 +48,12 @@ TValueOrError<FNewSpawnable, FText> FLevelSequenceEditorActorSpawner::CreateNewS
 	// If it's a blueprint, we need some special handling
 	else if (UBlueprint* SourceBlueprint = Cast<UBlueprint>(&SourceObject))
 	{
+		if (!OwnerMovieScene.GetClass()->IsChildOf(SourceBlueprint->GeneratedClass->ClassWithin))
+		{
+			ErrorText = FText::Format(LOCTEXT("ClassWithin", "Unable to add spawnable for class of type '{0}' since it has a required outer class '{1}'."), FText::FromString(SourceObject.GetName()), FText::FromString(SourceBlueprint->GeneratedClass->ClassWithin->GetName()));
+			return MakeError(ErrorText);
+		}
+
 		NewSpawnable.ObjectTemplate = NewObject<UObject>(&OwnerMovieScene, SourceBlueprint->GeneratedClass, TemplateName);
 	}
 
@@ -56,28 +63,31 @@ TValueOrError<FNewSpawnable, FText> FLevelSequenceEditorActorSpawner::CreateNewS
 		UActorFactory* FactoryToUse = ActorFactory ? ActorFactory : FActorFactoryAssetProxy::GetFactoryForAssetObject(&SourceObject);
 		if (!FactoryToUse)
 		{
-			ErrorText = FText::Format(LOCTEXT("CouldNotFindFactory", "Unable to create spawnable from  asset '{0}' - no valid factory could be found."), FText::FromString(SourceObject.GetName()));
+			ErrorText = FText::Format(LOCTEXT("CouldNotFindFactory", "Unable to create spawnable from asset '{0}' - no valid factory could be found."), FText::FromString(SourceObject.GetName()));
 		}
 
-		if (!FactoryToUse->CanCreateActorFrom(FAssetData(&SourceObject), ErrorText))
+		if (FactoryToUse)
 		{
-			if (!ErrorText.IsEmpty())
+			if (!FactoryToUse->CanCreateActorFrom(FAssetData(&SourceObject), ErrorText))
 			{
-				ErrorText = FText::Format(LOCTEXT("CannotCreateActorFromAsset_Ex", "Unable to create spawnable from  asset '{0}'. {1}."), FText::FromString(SourceObject.GetName()), ErrorText);
+				if (!ErrorText.IsEmpty())
+				{
+					ErrorText = FText::Format(LOCTEXT("CannotCreateActorFromAsset_Ex", "Unable to create spawnable from  asset '{0}'. {1}."), FText::FromString(SourceObject.GetName()), ErrorText);
+				}
+				else
+				{
+					ErrorText = FText::Format(LOCTEXT("CannotCreateActorFromAsset", "Unable to create spawnable from  asset '{0}'."), FText::FromString(SourceObject.GetName()));
+				}
 			}
-			else
-			{
-				ErrorText = FText::Format(LOCTEXT("CannotCreateActorFromAsset", "Unable to create spawnable from  asset '{0}'."), FText::FromString(SourceObject.GetName()));
-			}
-		}
 
-		AActor* Instance = FactoryToUse->CreateActor(&SourceObject, GWorld->PersistentLevel, FTransform(), RF_Transient, TemplateName );
-		Instance->bIsEditorPreviewActor = false;
-		NewSpawnable.ObjectTemplate = StaticDuplicateObject(Instance, &OwnerMovieScene, TemplateName, RF_AllFlags & ~RF_Transient);
+			AActor* Instance = FactoryToUse->CreateActor(&SourceObject, GWorld->PersistentLevel, FTransform(), RF_Transient, TemplateName );
+			Instance->bIsEditorPreviewActor = false;
+			NewSpawnable.ObjectTemplate = StaticDuplicateObject(Instance, &OwnerMovieScene, TemplateName, RF_AllFlags & ~RF_Transient);
 
-		const bool bNetForce = false;
-		const bool bShouldModifyLevel = false;
-		GWorld->DestroyActor(Instance, bNetForce, bShouldModifyLevel);
+			const bool bNetForce = false;
+			const bool bShouldModifyLevel = false;
+			GWorld->DestroyActor(Instance, bNetForce, bShouldModifyLevel);
+		}
 	}
 
 	if (!NewSpawnable.ObjectTemplate || !NewSpawnable.ObjectTemplate->IsA<AActor>())
@@ -148,9 +158,9 @@ bool FLevelSequenceEditorActorSpawner::CanSetupDefaultsForSpawnable(UObject* Spa
 	return FLevelSequenceActorSpawner::CanSetupDefaultsForSpawnable(SpawnedObject);
 }
 
-void FLevelSequenceEditorActorSpawner::SetupDefaultsForSpawnable(UObject* SpawnedObject, const FGuid& Guid, const FTransformData& TransformData, TSharedRef<ISequencer> Sequencer, USequencerSettings* Settings)
+void FLevelSequenceEditorActorSpawner::SetupDefaultsForSpawnable(UObject* SpawnedObject, const FGuid& Guid, const TOptional<FTransformData>& TransformData, TSharedRef<ISequencer> Sequencer, USequencerSettings* Settings)
 {
-	FTransformData DefaultTransform = TransformData;
+	TOptional<FTransformData> DefaultTransform = TransformData;
 
 	AActor* SpawnedActor = Cast<AActor>(SpawnedObject);
 	if (SpawnedActor)
@@ -160,10 +170,11 @@ void FLevelSequenceEditorActorSpawner::SetupDefaultsForSpawnable(UObject* Spawne
 		{
 			PlaceActorInFrontOfCamera(SpawnedActor);
 		}
-		DefaultTransform.Translation = SpawnedActor->GetActorLocation();
-		DefaultTransform.Rotation = SpawnedActor->GetActorRotation();
-		DefaultTransform.Scale = FVector(1.0f, 1.0f, 1.0f);
-		DefaultTransform.bValid = true;
+		DefaultTransform.Reset();
+		DefaultTransform.Emplace();
+		DefaultTransform->Translation = SpawnedActor->GetActorLocation();
+		DefaultTransform->Rotation = SpawnedActor->GetActorRotation();
+		DefaultTransform->Scale = FVector(1.0f, 1.0f, 1.0f);
 
 		Sequencer->OnActorAddedToSequencer().Broadcast(SpawnedActor, Guid);
 
@@ -189,14 +200,17 @@ void FLevelSequenceEditorActorSpawner::SetupDefaultsForSpawnable(UObject* Spawne
 	if (SpawnTrack)
 	{
 		UMovieSceneBoolSection* SpawnSection = Cast<UMovieSceneBoolSection>(SpawnTrack->CreateNewSection());
-		SpawnSection->SetDefault(true);
-		SpawnSection->SetIsInfinite(Sequencer->GetInfiniteKeyAreas());
+		SpawnSection->GetChannel().SetDefault(true);
+		if (Sequencer->GetInfiniteKeyAreas())
+		{
+			SpawnSection->SetRange(TRange<FFrameNumber>::All());
+		}
 		SpawnTrack->AddSection(*SpawnSection);
 		SpawnTrack->SetObjectId(Guid);
 	}
 
 	// Ensure it will spawn in the right place
-	if (DefaultTransform.bValid)
+	if (DefaultTransform.IsSet())
 	{
 		UMovieScene3DTransformTrack* TransformTrack = Cast<UMovieScene3DTransformTrack>(OwnerMovieScene->FindTrack(UMovieScene3DTransformTrack::StaticClass(), Guid, "Transform"));
 		if (!TransformTrack)
@@ -206,33 +220,39 @@ void FLevelSequenceEditorActorSpawner::SetupDefaultsForSpawnable(UObject* Spawne
 
 		if (TransformTrack)
 		{
-			const bool bUnwindRotation = false;
-
-			EMovieSceneKeyInterpolation Interpolation = Sequencer->GetKeyInterpolation();
-
 			const TArray<UMovieSceneSection*>& Sections = TransformTrack->GetAllSections();
 			if (!Sections.Num())
 			{
 				TransformTrack->AddSection(*TransformTrack->CreateNewSection());
 			}
 
+			FVector Location = DefaultTransform->Translation;
+			FVector Rotation = DefaultTransform->Rotation.Euler();
+			FVector Scale    = DefaultTransform->Scale;
+
 			for (UMovieSceneSection* Section : Sections)
 			{
 				UMovieScene3DTransformSection* TransformSection = CastChecked<UMovieScene3DTransformSection>(Section);
 
-				TransformSection->SetDefault(FTransformKey(EKey3DTransformChannel::Translation, EAxis::X, DefaultTransform.Translation.X, bUnwindRotation));
-				TransformSection->SetDefault(FTransformKey(EKey3DTransformChannel::Translation, EAxis::Y, DefaultTransform.Translation.Y, bUnwindRotation));
-				TransformSection->SetDefault(FTransformKey(EKey3DTransformChannel::Translation, EAxis::Z, DefaultTransform.Translation.Z, bUnwindRotation));
-
-				TransformSection->SetDefault(FTransformKey(EKey3DTransformChannel::Rotation, EAxis::X, DefaultTransform.Rotation.Euler().X, bUnwindRotation));
-				TransformSection->SetDefault(FTransformKey(EKey3DTransformChannel::Rotation, EAxis::Y, DefaultTransform.Rotation.Euler().Y, bUnwindRotation));
-				TransformSection->SetDefault(FTransformKey(EKey3DTransformChannel::Rotation, EAxis::Z, DefaultTransform.Rotation.Euler().Z, bUnwindRotation));
-
-				TransformSection->SetDefault(FTransformKey(EKey3DTransformChannel::Scale, EAxis::X, DefaultTransform.Scale.X, bUnwindRotation));
-				TransformSection->SetDefault(FTransformKey(EKey3DTransformChannel::Scale, EAxis::Y, DefaultTransform.Scale.Y, bUnwindRotation));
-				TransformSection->SetDefault(FTransformKey(EKey3DTransformChannel::Scale, EAxis::Z, DefaultTransform.Scale.Z, bUnwindRotation));
-
-				TransformSection->SetIsInfinite(Sequencer->GetInfiniteKeyAreas());
+				// Set the section to be infinite if necessary
+				if (Sequencer->GetInfiniteKeyAreas())
+				{
+					TransformSection->SetRange(TRange<FFrameNumber>::All());
+				}
+
+				// Set the section's default values to this default transform
+				TArrayView<FMovieSceneFloatChannel*> FloatChannels = TransformSection->GetChannelProxy().GetChannels<FMovieSceneFloatChannel>();
+				FloatChannels[0]->SetDefault(Location.X);
+				FloatChannels[1]->SetDefault(Location.Y);
+				FloatChannels[2]->SetDefault(Location.Z);
+
+				FloatChannels[3]->SetDefault(Rotation.X);
+				FloatChannels[4]->SetDefault(Rotation.Y);
+				FloatChannels[5]->SetDefault(Rotation.Z);
+
+				FloatChannels[6]->SetDefault(Scale.X);
+				FloatChannels[7]->SetDefault(Scale.Y);
+				FloatChannels[8]->SetDefault(Scale.Z);
 			}
 		}
 	}
