Media Sections now set the Looping flag if their sections go past their duration. This aids in some media players (like EXR) that can use the looping flag to pre-cache data properly.

diff --git a/Engine/Plugins/Media/MediaCompositing/Source/MediaCompositing/Private/MovieScene/MovieSceneMediaTemplate.cpp b/Engine/Plugins/Media/MediaCompositing/Source/MediaCompositing/Private/MovieScene/MovieSceneMediaTemplate.cpp
--- a/Engine/Plugins/Media/MediaCompositing/Source/MediaCompositing/Private/MovieScene/MovieSceneMediaTemplate.cpp
+++ b/Engine/Plugins/Media/MediaCompositing/Source/MediaCompositing/Private/MovieScene/MovieSceneMediaTemplate.cpp
@@ -149,6 +149,7 @@ FMovieSceneMediaSectionTemplate::FMovieSceneMediaSectionTemplate(const UMovieSce
 	Params.MediaSoundComponent = InSection.MediaSoundComponent;
 	Params.MediaSource = InSection.GetMediaSource();
 	Params.MediaTexture = InSection.MediaTexture;
+	Params.bLooping = InSection.bLooping;
 
 	if (InSection.HasStartFrame())
 	{
@@ -269,6 +270,8 @@ void FMovieSceneMediaSectionTemplate::Initialize(const FMovieSceneEvaluationOper
 			Params.MediaTexture->SetMediaPlayer(nullptr);
 		}
 	}
+
+	MediaPlayer->SetLooping(Params.bLooping);
 }
 
 
diff --git a/Engine/Plugins/Media/MediaCompositing/Source/MediaCompositing/Private/MovieScene/MovieSceneMediaTemplate.h b/Engine/Plugins/Media/MediaCompositing/Source/MediaCompositing/Private/MovieScene/MovieSceneMediaTemplate.h
--- a/Engine/Plugins/Media/MediaCompositing/Source/MediaCompositing/Private/MovieScene/MovieSceneMediaTemplate.h
+++ b/Engine/Plugins/Media/MediaCompositing/Source/MediaCompositing/Private/MovieScene/MovieSceneMediaTemplate.h
@@ -36,6 +36,9 @@ struct FMovieSceneMediaSectionParams
 
 	UPROPERTY()
 	FFrameNumber SectionEndFrame;
+
+	UPROPERTY()
+	bool bLooping;
 };
 
 
diff --git a/Engine/Plugins/Media/MediaCompositing/Source/MediaCompositing/Public/MovieSceneMediaSection.h b/Engine/Plugins/Media/MediaCompositing/Source/MediaCompositing/Public/MovieSceneMediaSection.h
--- a/Engine/Plugins/Media/MediaCompositing/Source/MediaCompositing/Public/MovieSceneMediaSection.h
+++ b/Engine/Plugins/Media/MediaCompositing/Source/MediaCompositing/Public/MovieSceneMediaSection.h
@@ -22,7 +22,11 @@ public:
 
 	GENERATED_BODY()
 
-	/** The media source proxy to use. */
+	/** Should the media player be set to loop? This can be helpful for media formats that can use this information (such as exr sequences) to pre-cache the starting data when nearing the end of playback. Media playback is still limited by length of section. */
+	UPROPERTY(EditAnywhere, Category = "Media")
+	bool bLooping;
+
+	/** The media source proxy to use. This allows the usage of a lower resolution (or otherwise modified) version during development. */
 	UPROPERTY(EditAnywhere, Category="Media")
 	FString Proxy;
 
@@ -33,7 +37,6 @@ public:
 	/** The media texture that receives the track's video output. */
 	UPROPERTY(EditAnywhere, Category="Media")
 	UMediaTexture* MediaTexture;
-
 public:
 
 	/**
