Sequencer: Fix media timing - seek to media time before blocking

diff --git a/Engine/Plugins/Media/MediaCompositing/Source/MediaCompositing/Private/MovieScene/MovieSceneMediaTemplate.cpp b/Engine/Plugins/Media/MediaCompositing/Source/MediaCompositing/Private/MovieScene/MovieSceneMediaTemplate.cpp
--- a/Engine/Plugins/Media/MediaCompositing/Source/MediaCompositing/Private/MovieScene/MovieSceneMediaTemplate.cpp
+++ b/Engine/Plugins/Media/MediaCompositing/Source/MediaCompositing/Private/MovieScene/MovieSceneMediaTemplate.cpp
@@ -114,7 +114,7 @@ struct FMediaSectionExecutionToken
 				MediaPlayer->Seek(MediaTime);
 				MediaPlayer->SetRate(1.0f);
 			}
-			else if (Context.HasJumped())
+			else 
 			{
 				MediaPlayer->Seek(MediaTime);
 			}
