Back out revision 26 from //UE4/Dev-Sequencer/Engine/Plugins/Media/MediaCompositing/Source/MediaCompositing/Private/MovieScene/MovieSceneMediaTemplate.cpp

diff --git a/Engine/Plugins/Media/MediaCompositing/Source/MediaCompositing/Private/MovieScene/MovieSceneMediaTemplate.cpp b/Engine/Plugins/Media/MediaCompositing/Source/MediaCompositing/Private/MovieScene/MovieSceneMediaTemplate.cpp
--- a/Engine/Plugins/Media/MediaCompositing/Source/MediaCompositing/Private/MovieScene/MovieSceneMediaTemplate.cpp
+++ b/Engine/Plugins/Media/MediaCompositing/Source/MediaCompositing/Private/MovieScene/MovieSceneMediaTemplate.cpp
@@ -114,7 +114,7 @@ struct FMediaSectionExecutionToken
 				MediaPlayer->Seek(MediaTime);
 				MediaPlayer->SetRate(1.0f);
 			}
-			else 
+			else if (Context.HasJumped())
 			{
 				MediaPlayer->Seek(MediaTime);
 			}
