Changing "WhitelistTargets" to be an array instead of a string

diff --git a/Engine/Plugins/Media/MediaCompositing/MediaCompositing.uplugin b/Engine/Plugins/Media/MediaCompositing/MediaCompositing.uplugin
--- a/Engine/Plugins/Media/MediaCompositing/MediaCompositing.uplugin
+++ b/Engine/Plugins/Media/MediaCompositing/MediaCompositing.uplugin
@@ -30,7 +30,7 @@
 	[
 		{
 			"Name": "ImgMedia",
-			"WhitelistTargets": "Editor",
+			"WhitelistTargets": [ "Editor" ],
 			"Enabled": true
 		}
 	]
