Sequencer: Added the ability to specify an external clock source interface to sequencer

This allows us to drive sequencer time from arbitrary external sources such as media players or other clocks.
Added a new media component which owns a media player and texture, such that media players can work nicely with multiplayer PIE
Added a new media track that can play back media sources on any UMediaPlayer* property - primitive for the time being (only looping is supported, no start time offset etc). This track is currently written to support real-time playback rather than offline, frame-accurate playback

diff --git a/Engine/Plugins/Media/MediaCompositing/Source/MediaCompositing/Private/MovieScene/MovieSceneMediaPlayerPropertySection.cpp b/Engine/Plugins/Media/MediaCompositing/Source/MediaCompositing/Private/MovieScene/MovieSceneMediaPlayerPropertySection.cpp
new file mode 100644
--- /dev/null
+++ b/Engine/Plugins/Media/MediaCompositing/Source/MediaCompositing/Private/MovieScene/MovieSceneMediaPlayerPropertySection.cpp
@@ -0,0 +1,16 @@
+// Copyright 1998-2020 Epic Games, Inc. All Rights Reserved.
+
+#include "MovieSceneMediaPlayerPropertySection.h"
+#include "MovieScene.h"
+
+UMovieSceneMediaPlayerPropertySection::UMovieSceneMediaPlayerPropertySection(const FObjectInitializer& ObjectInitializer)
+	: Super(ObjectInitializer)
+{
+	EvalOptions.CompletionMode = EMovieSceneCompletionMode::RestoreState;
+	bLoop = false;
+
+	UMovieScene* Outer = GetTypedOuter<UMovieScene>();
+	FFrameRate TickResolution = Outer ? Outer->GetTickResolution() : FFrameRate(24, 1);
+
+	SetPreRollFrames( (0.5 * TickResolution).RoundToFrame().Value );
+}
diff --git a/Engine/Plugins/Media/MediaCompositing/Source/MediaCompositing/Private/MovieScene/MovieSceneMediaPlayerPropertyTemplate.cpp b/Engine/Plugins/Media/MediaCompositing/Source/MediaCompositing/Private/MovieScene/MovieSceneMediaPlayerPropertyTemplate.cpp
new file mode 100644
--- /dev/null
+++ b/Engine/Plugins/Media/MediaCompositing/Source/MediaCompositing/Private/MovieScene/MovieSceneMediaPlayerPropertyTemplate.cpp
@@ -0,0 +1,228 @@
+// Copyright 1998-2020 Epic Games, Inc. All Rights Reserved.
+
+#include "MovieSceneMediaPlayerPropertyTemplate.h"
+
+#include "MediaPlayer.h"
+#include "MediaSource.h"
+#include "MediaPlaylist.h"
+
+#include "MovieSceneMediaPlayerPropertySection.h"
+#include "MovieSceneMediaPlayerPropertyTrack.h"
+
+
+struct FMediaPlayerPreAnimatedState : IMovieScenePreAnimatedToken
+{
+	virtual void RestoreState(UObject& Object, IMovieScenePlayer& Player)
+	{
+		CastChecked<UMediaPlayer>(&Object)->Close();
+	}
+};
+
+struct FMediaPlayerSectionPreRollExecutionToken : IMovieSceneExecutionToken
+{
+	FMediaPlayerSectionPreRollExecutionToken(UMediaSource* InMediaSource, FTimespan InStartTimeSeconds, bool bInLoop)
+		: MediaSource(InMediaSource)
+		, StartTime(InStartTimeSeconds)
+		, bLoop(bInLoop)
+	{}
+
+	virtual void Execute(const FMovieSceneContext& Context, const FMovieSceneEvaluationOperand& Operand, FPersistentEvaluationData& PersistentData, IMovieScenePlayer& Player) override
+	{
+		using namespace PropertyTemplate;
+		const FSectionData& SectionData = PersistentData.GetSectionData<FSectionData>();
+
+		const FTimespan TwoFrames = FTimespan::FromSeconds(1 / 12.f);
+
+		for (TWeakObjectPtr<> WeakObject : Player.FindBoundObjects(Operand))
+		{
+			UObject* BoundObject = WeakObject.Get();
+			if (!BoundObject)
+			{
+				continue;
+			}
+
+			UMediaPlayer* MediaPlayer = SectionData.PropertyBindings->GetCurrentValue<UMediaPlayer*>(*BoundObject);
+			if (!MediaPlayer)
+			{
+				continue;
+			}
+
+			// Ensure the media player is playing this media
+			const int32 PlaylistIndex = MediaPlayer->GetPlaylistIndex();
+			if (PlaylistIndex == INDEX_NONE || MediaPlayer->GetPlaylistRef().Get(PlaylistIndex) != MediaSource)
+			{
+				FMediaPlayerOptions Options;
+				Options.SeekTime   = StartTime;
+				Options.PlayOnOpen = EMediaPlayerOptionBooleanOverride::Disabled;
+				Options.Loop       = bLoop ? EMediaPlayerOptionBooleanOverride::Enabled : EMediaPlayerOptionBooleanOverride::Disabled;
+				MediaPlayer->OpenSourceWithOptions(MediaSource, Options);
+
+				continue;
+			}
+
+			const FTimespan MediaDuration = MediaPlayer->GetDuration();
+			if (!MediaDuration.IsZero() && StartTime < MediaDuration)
+			{
+				if (FMath::Abs(MediaPlayer->GetTime() - StartTime) >= TwoFrames)
+				{
+					MediaPlayer->Seek(StartTime);
+				}
+				MediaPlayer->SetRate(0.0f);
+			}
+		}
+	}
+
+private:
+
+	UMediaSource* MediaSource;
+	FTimespan StartTime;
+	bool bLoop;
+};
+
+
+struct FMediaPlayerSectionExecutionToken : IMovieSceneExecutionToken
+{
+	FMediaPlayerSectionExecutionToken(UMediaSource* InMediaSource, FTimespan InCurrentTime, bool bInLoop)
+		: MediaSource(InMediaSource)
+		, CurrentTime(InCurrentTime)
+		, bLoop(bInLoop)
+	{ }
+
+	virtual void Execute(const FMovieSceneContext& Context, const FMovieSceneEvaluationOperand& Operand, FPersistentEvaluationData& PersistentData, IMovieScenePlayer& Player) override
+	{
+		using namespace PropertyTemplate;
+		const FSectionData& SectionData = PersistentData.GetSectionData<FSectionData>();
+
+		const FTimespan TwoFrames = FTimespan::FromSeconds(1 / 12.f);
+
+		for (TWeakObjectPtr<> WeakObject : Player.FindBoundObjects(Operand))
+		{
+			UObject* BoundObject = WeakObject.Get();
+			if (!BoundObject)
+			{
+				continue;
+			}
+
+			UMediaPlayer* MediaPlayer = SectionData.PropertyBindings->GetCurrentValue<UMediaPlayer*>(*BoundObject);
+			if (!MediaPlayer)
+			{
+				continue;
+			}
+
+			Player.SavePreAnimatedState(*MediaPlayer, TMovieSceneAnimTypeID<FMediaPlayerSectionExecutionToken>(), TStatelessPreAnimatedTokenProducer<FMediaPlayerPreAnimatedState>(), PersistentData.GetSectionKey());
+
+			// Ensure the media player is playing this media
+			const int32 PlaylistIndex = MediaPlayer->GetPlaylistIndex();
+			if (PlaylistIndex == INDEX_NONE || !MediaPlayer->GetPlaylist() || MediaPlayer->GetPlaylist()->Get(PlaylistIndex) != MediaSource)
+			{
+				FMediaPlayerOptions Options;
+				Options.SeekTime   = CurrentTime;
+				Options.PlayOnOpen = Context.GetStatus() == EMovieScenePlayerStatus::Playing ? EMediaPlayerOptionBooleanOverride::Enabled : EMediaPlayerOptionBooleanOverride::Disabled;
+				Options.Loop       = bLoop ? EMediaPlayerOptionBooleanOverride::Enabled : EMediaPlayerOptionBooleanOverride::Disabled;
+				MediaPlayer->OpenSourceWithOptions(MediaSource, Options);
+			}
+
+			const FTimespan MediaDuration = MediaPlayer->GetDuration();
+			if (MediaDuration.IsZero())
+			{
+				continue;
+			}
+
+
+			if (bLoop)
+			{
+				CurrentTime = CurrentTime % MediaDuration;
+			}
+			else if (CurrentTime > MediaDuration)
+			{
+				MediaPlayer->Pause();
+				continue;
+			}
+
+			const FTimespan MediaTime = MediaPlayer->GetTime();
+
+			if (Context.GetStatus() == EMovieScenePlayerStatus::Playing)
+			{
+				if (!MediaPlayer->IsPlaying())
+				{
+					if (FMath::Abs(MediaTime - CurrentTime) >= TwoFrames)
+					{
+						MediaPlayer->Seek(CurrentTime);
+					}
+
+					if (MediaPlayer->GetRate() != 1.0f)
+					{
+						MediaPlayer->SetRate(1.0f);
+					}
+				}
+				else if (Context.HasJumped())
+				{
+					if (FMath::Abs(MediaTime - CurrentTime) >= TwoFrames)
+					{
+						MediaPlayer->Seek(CurrentTime);
+					}
+				}
+			}
+			else
+			{
+				if (FMath::Abs(MediaTime - CurrentTime) >= TwoFrames) 
+				{
+					MediaPlayer->Seek(CurrentTime);
+				}
+
+				// Seek can sometimes cause playback to start - ensure that doesn't happen
+				MediaPlayer->Pause();
+			}
+		}
+	}
+
+private:
+
+	UMediaSource* MediaSource;
+	FTimespan CurrentTime;
+	bool bLoop;
+};
+
+
+FMovieSceneMediaPlayerPropertySectionTemplate::FMovieSceneMediaPlayerPropertySectionTemplate(const UMovieSceneMediaPlayerPropertySection* InSection, const UMovieSceneMediaPlayerPropertyTrack* InTrack)
+	: FMovieScenePropertySectionTemplate(InTrack->GetPropertyName(), InTrack->GetPropertyPath())
+{
+	MediaSource = InSection->MediaSource;
+	bLoop = InSection->bLoop;
+
+	if (InSection->HasStartFrame())
+	{
+		SectionStartFrame = InSection->GetRange().GetLowerBoundValue();
+	}
+}
+
+void FMovieSceneMediaPlayerPropertySectionTemplate::Evaluate(const FMovieSceneEvaluationOperand& Operand, const FMovieSceneContext& Context, const FPersistentEvaluationData& PersistentData, FMovieSceneExecutionTokens& ExecutionTokens) const
+{
+	// Can't do anything without a media source
+	if (!MediaSource)
+	{
+		return;
+	}
+
+	// @todo: account for video time dilation if/when these are added
+	if (Context.IsPreRoll())
+	{
+		const FFrameRate   FrameRate  = Context.GetFrameRate();
+		const FFrameNumber StartFrame = Context.HasPreRollEndTime() ? Context.GetPreRollEndFrame() - SectionStartFrame : 0;
+
+		const int64 DenominatorTicks = FrameRate.Denominator * ETimespan::TicksPerSecond;
+		const int64 StartTicks = FMath::DivideAndRoundNearest(int64(StartFrame.Value * DenominatorTicks), int64(FrameRate.Numerator));
+
+		ExecutionTokens.Add(FMediaPlayerSectionPreRollExecutionToken(MediaSource, FTimespan(StartTicks), bLoop));
+	}
+	else
+	{
+		const FFrameRate FrameRate = Context.GetFrameRate();
+		const FFrameTime FrameTime(Context.GetTime().FrameNumber - SectionStartFrame);
+		const int64 DenominatorTicks = FrameRate.Denominator * ETimespan::TicksPerSecond;
+		const int64 FrameTicks = FMath::DivideAndRoundNearest(int64(FrameTime.GetFrame().Value * DenominatorTicks), int64(FrameRate.Numerator));
+		const int64 FrameSubTicks = FMath::DivideAndRoundNearest(int64(FrameTime.GetSubFrame() * DenominatorTicks), int64(FrameRate.Numerator));
+
+		ExecutionTokens.Add(FMediaPlayerSectionExecutionToken(MediaSource, FTimespan(FrameTicks + FrameSubTicks), bLoop));
+	}
+}
\ No newline at end of file
diff --git a/Engine/Plugins/Media/MediaCompositing/Source/MediaCompositing/Private/MovieScene/MovieSceneMediaPlayerPropertyTemplate.h b/Engine/Plugins/Media/MediaCompositing/Source/MediaCompositing/Private/MovieScene/MovieSceneMediaPlayerPropertyTemplate.h
new file mode 100644
--- /dev/null
+++ b/Engine/Plugins/Media/MediaCompositing/Source/MediaCompositing/Private/MovieScene/MovieSceneMediaPlayerPropertyTemplate.h
@@ -0,0 +1,40 @@
+// Copyright 1998-2020 Epic Games, Inc. All Rights Reserved.
+
+#pragma once
+
+#include "Evaluation/MovieScenePropertyTemplate.h"
+
+#include "MovieSceneMediaPlayerPropertyTemplate.generated.h"
+
+class UMediaSource;
+class UMovieSceneMediaPlayerPropertyTrack;
+class UMovieSceneMediaPlayerPropertySection;
+
+
+USTRUCT()
+struct FMovieSceneMediaPlayerPropertySectionTemplate
+	: public FMovieScenePropertySectionTemplate
+{
+	GENERATED_BODY()
+
+	/** Default constructor. */
+	FMovieSceneMediaPlayerPropertySectionTemplate() { }
+	FMovieSceneMediaPlayerPropertySectionTemplate(const UMovieSceneMediaPlayerPropertySection* InSection, const UMovieSceneMediaPlayerPropertyTrack* InTrack);
+
+public:
+
+	virtual void Evaluate(const FMovieSceneEvaluationOperand& Operand, const FMovieSceneContext& Context, const FPersistentEvaluationData& PersistentData, FMovieSceneExecutionTokens& ExecutionTokens) const override;
+	virtual UScriptStruct& GetScriptStructImpl() const override { return *StaticStruct(); }
+	virtual void SetupOverrides() override { EnableOverrides(RequiresSetupFlag); }
+
+private:
+
+	UPROPERTY()
+	UMediaSource* MediaSource;
+
+	UPROPERTY()
+	FFrameNumber SectionStartFrame;
+
+	UPROPERTY()
+	bool bLoop;
+};
diff --git a/Engine/Plugins/Media/MediaCompositing/Source/MediaCompositing/Private/MovieScene/MovieSceneMediaPlayerPropertyTrack.cpp b/Engine/Plugins/Media/MediaCompositing/Source/MediaCompositing/Private/MovieScene/MovieSceneMediaPlayerPropertyTrack.cpp
new file mode 100644
--- /dev/null
+++ b/Engine/Plugins/Media/MediaCompositing/Source/MediaCompositing/Private/MovieScene/MovieSceneMediaPlayerPropertyTrack.cpp
@@ -0,0 +1,39 @@
+// Copyright 1998-2020 Epic Games, Inc. All Rights Reserved.
+
+#include "MovieSceneMediaPlayerPropertyTrack.h"
+
+#include "MediaSource.h"
+#include "MovieScene.h"
+#include "MovieSceneMediaPlayerPropertySection.h"
+#include "MovieSceneMediaPlayerPropertyTemplate.h"
+#include "UObject/UObjectGlobals.h"
+
+
+UMovieSceneMediaPlayerPropertyTrack::UMovieSceneMediaPlayerPropertyTrack(const FObjectInitializer& ObjectInitializer)
+	: Super(ObjectInitializer)
+{
+	EvalOptions.bCanEvaluateNearestSection = false;
+	EvalOptions.bEvalNearestSection = false;
+	EvalOptions.bEvaluateInPreroll = true;
+	EvalOptions.bEvaluateInPostroll = true;
+
+#if WITH_EDITORONLY_DATA
+	TrackTint = FColor(0, 0, 0, 200);
+#endif
+}
+
+
+bool UMovieSceneMediaPlayerPropertyTrack::SupportsType(TSubclassOf<UMovieSceneSection> SectionClass) const
+{
+	return SectionClass == UMovieSceneMediaPlayerPropertySection::StaticClass();
+}
+
+UMovieSceneSection* UMovieSceneMediaPlayerPropertyTrack::CreateNewSection()
+{
+	return NewObject<UMovieSceneMediaPlayerPropertySection>(this, NAME_None, RF_Transactional);
+}
+
+FMovieSceneEvalTemplatePtr UMovieSceneMediaPlayerPropertyTrack::CreateTemplateForSection(const UMovieSceneSection& InSection) const
+{
+	return FMovieSceneMediaPlayerPropertySectionTemplate(CastChecked<const UMovieSceneMediaPlayerPropertySection>(&InSection), this);
+}
diff --git a/Engine/Plugins/Media/MediaCompositing/Source/MediaCompositing/Public/MovieSceneMediaPlayerPropertySection.h b/Engine/Plugins/Media/MediaCompositing/Source/MediaCompositing/Public/MovieSceneMediaPlayerPropertySection.h
new file mode 100644
--- /dev/null
+++ b/Engine/Plugins/Media/MediaCompositing/Source/MediaCompositing/Public/MovieSceneMediaPlayerPropertySection.h
@@ -0,0 +1,30 @@
+// Copyright 1998-2020 Epic Games, Inc. All Rights Reserved.
+
+#pragma once
+
+#include "MovieSceneSection.h"
+#include "MovieSceneMediaPlayerPropertySection.generated.h"
+
+class UMediaSource;
+
+/**
+ * Implements a movie scene section for media playback on a UMediaPlayer.
+ */
+UCLASS(MinimalAPI)
+class UMovieSceneMediaPlayerPropertySection
+	: public UMovieSceneSection
+{
+public:
+
+	GENERATED_BODY()
+
+	/** The source to play with this video track. */
+	UPROPERTY(EditAnywhere, BlueprintReadWrite, Category="Media")
+	UMediaSource* MediaSource;
+
+	/** Whether to loop this video. */
+	UPROPERTY(EditAnywhere, BlueprintReadWrite, Category="Media")
+	bool bLoop;
+
+	UMovieSceneMediaPlayerPropertySection(const FObjectInitializer& ObjectInitializer);
+};
diff --git a/Engine/Plugins/Media/MediaCompositing/Source/MediaCompositing/Public/MovieSceneMediaPlayerPropertyTrack.h b/Engine/Plugins/Media/MediaCompositing/Source/MediaCompositing/Public/MovieSceneMediaPlayerPropertyTrack.h
new file mode 100644
--- /dev/null
+++ b/Engine/Plugins/Media/MediaCompositing/Source/MediaCompositing/Public/MovieSceneMediaPlayerPropertyTrack.h
@@ -0,0 +1,25 @@
+// Copyright 1998-2020 Epic Games, Inc. All Rights Reserved.
+
+#pragma once
+
+#include "UObject/ObjectMacros.h"
+#include "Tracks/MovieScenePropertyTrack.h"
+
+#include "MovieSceneMediaPlayerPropertyTrack.generated.h"
+
+UCLASS(MinimalAPI)
+class UMovieSceneMediaPlayerPropertyTrack
+	: public UMovieScenePropertyTrack
+{
+public:
+
+	GENERATED_BODY()
+
+	UMovieSceneMediaPlayerPropertyTrack(const FObjectInitializer& ObjectInitializer);
+
+public:
+
+	virtual bool SupportsType(TSubclassOf<UMovieSceneSection> SectionClass) const override;
+	virtual UMovieSceneSection* CreateNewSection() override;
+	virtual FMovieSceneEvalTemplatePtr CreateTemplateForSection(const UMovieSceneSection& InSection) const override;
+};
diff --git a/Engine/Plugins/Media/MediaCompositing/Source/MediaCompositingEditor/Private/MediaCompositingEditor.cpp b/Engine/Plugins/Media/MediaCompositing/Source/MediaCompositingEditor/Private/MediaCompositingEditor.cpp
--- a/Engine/Plugins/Media/MediaCompositing/Source/MediaCompositingEditor/Private/MediaCompositingEditor.cpp
+++ b/Engine/Plugins/Media/MediaCompositing/Source/MediaCompositingEditor/Private/MediaCompositingEditor.cpp
@@ -8,6 +8,7 @@
 #include "MediaCompositingEditorStyle.h"
 #include "MediaSequenceRecorderExtender.h"
 #include "MediaTrackEditor.h"
+#include "MediaPlayerPropertyTrackEditor.h"
 
 
 #define LOCTEXT_NAMESPACE "MediaCompositingEditorModule"
@@ -27,6 +28,7 @@ public:
 
 		ISequencerModule& SequencerModule = FModuleManager::Get().LoadModuleChecked<ISequencerModule>("Sequencer");
 		TrackEditorBindingHandle = SequencerModule.RegisterPropertyTrackEditor<FMediaTrackEditor>();
+		PropertyTrackEditorBindingHandle = SequencerModule.RegisterPropertyTrackEditor<FMediaPlayerPropertyTrackEditor>();
 
 		ISequenceRecorder& SequenceRecorder = FModuleManager::LoadModuleChecked<ISequenceRecorder>("SequenceRecorder");
 		RecorderExtender = MakeShared<FMediaSequenceRecorderExtender>();
@@ -49,10 +51,12 @@ public:
 		if (SequencerModulePtr)
 		{
 			SequencerModulePtr->UnRegisterTrackEditor(TrackEditorBindingHandle);
+			SequencerModulePtr->UnRegisterTrackEditor(PropertyTrackEditorBindingHandle);
 		}
 	}
 
 	FDelegateHandle TrackEditorBindingHandle;
+	FDelegateHandle PropertyTrackEditorBindingHandle;
 	TSharedPtr<FMediaSequenceRecorderExtender> RecorderExtender;
 };
 
diff --git a/Engine/Plugins/Media/MediaCompositing/Source/MediaCompositingEditor/Private/Sequencer/MediaPlayerPropertyTrackEditor.cpp b/Engine/Plugins/Media/MediaCompositing/Source/MediaCompositingEditor/Private/Sequencer/MediaPlayerPropertyTrackEditor.cpp
new file mode 100644
--- /dev/null
+++ b/Engine/Plugins/Media/MediaCompositing/Source/MediaCompositingEditor/Private/Sequencer/MediaPlayerPropertyTrackEditor.cpp
@@ -0,0 +1,239 @@
+// Copyright 1998-2020 Epic Games, Inc. All Rights Reserved.
+
+#include "MediaPlayerPropertyTrackEditor.h"
+
+#include "ContentBrowserModule.h"
+#include "EditorStyleSet.h"
+#include "Framework/Application/SlateApplication.h"
+#include "Framework/MultiBox/MultiBoxBuilder.h"
+#include "IContentBrowserSingleton.h"
+#include "ISequencer.h"
+#include "SequencerSectionPainter.h"
+#include "MovieSceneMediaPlayerPropertySection.h"
+#include "MovieSceneMediaPlayerPropertyTrack.h"
+#include "SequencerUtilities.h"
+#include "Widgets/SBoxPanel.h"
+#include "Widgets/Layout/SBox.h"
+#include "MediaSource.h"
+
+#include "MediaThumbnailSection.h"
+
+#define LOCTEXT_NAMESPACE "FMediaPlayerPropertyTrackEditor"
+
+
+struct FMediaPlayerPropertySection : FSequencerSection
+{
+	FMediaPlayerPropertySection(UMovieSceneMediaPlayerPropertySection* InSection)
+		: FSequencerSection(*InSection)
+	{}
+
+	virtual FMargin GetContentPadding() const override
+	{
+		return FMargin(8.0f, 15.0f);
+	}
+
+	virtual float GetSectionHeight() const override
+	{
+		const float InnerHeight = FEditorStyle::GetFontStyle("NormalFont").Size + 8.f;
+		return InnerHeight + 2 * 9.0f; // make space for the film border
+	}
+
+	virtual FText GetSectionTitle() const override
+	{
+		UMovieSceneMediaPlayerPropertySection* MediaSection = CastChecked<UMovieSceneMediaPlayerPropertySection>(WeakSection.Get());
+		if (MediaSection->MediaSource == nullptr)
+		{
+			return LOCTEXT("NoSequence", "Empty");
+		}
+
+		return FText::FromString(MediaSection->MediaSource->GetFName().ToString());
+	}
+
+	virtual int32 OnPaintSection(FSequencerSectionPainter& InPainter) const override
+	{
+		// draw background
+		InPainter.LayerId = InPainter.PaintSectionBackground();
+
+		FVector2D SectionSize = InPainter.SectionGeometry.GetLocalSize();
+		FSlateClippingZone ClippingZone(InPainter.SectionClippingRect.InsetBy(FMargin(1.0f)));
+		
+		InPainter.DrawElements.PushClip(ClippingZone);
+		{
+			static const FSlateBrush* FilmBorder = FEditorStyle::GetBrush("Sequencer.Section.FilmBorder");
+
+			// draw top film border
+			FSlateDrawElement::MakeBox(
+				InPainter.DrawElements,
+				InPainter.LayerId++,
+				InPainter.SectionGeometry.ToPaintGeometry(FVector2D(SectionSize.X - 2.0f, 7.0f), FSlateLayoutTransform(FVector2D(1.0f, 4.0f))),
+				FilmBorder,
+				InPainter.bParentEnabled ? ESlateDrawEffect::None : ESlateDrawEffect::DisabledEffect
+			);
+
+			// draw bottom film border
+			FSlateDrawElement::MakeBox(
+				InPainter.DrawElements,
+				InPainter.LayerId++,
+				InPainter.SectionGeometry.ToPaintGeometry(FVector2D(SectionSize.X - 2.0f, 7.0f), FSlateLayoutTransform(FVector2D(1.0f, SectionSize.Y - 11.0f))),
+				FilmBorder,
+				InPainter.bParentEnabled ? ESlateDrawEffect::None : ESlateDrawEffect::DisabledEffect
+			);
+		}
+		InPainter.DrawElements.PopClip();
+
+		return InPainter.LayerId;
+	}
+};
+
+
+
+void FMediaPlayerPropertyTrackEditor::OnAnimatedPropertyChanged(const FPropertyChangedParams& PropertyChangedParams)
+{
+	// Get the movie scene we want to autokey
+	UMovieSceneSequence* MovieSceneSequence = GetMovieSceneSequence();
+	if (!MovieSceneSequence)
+	{
+		return;
+	}
+
+	if (PropertyChangedParams.KeyMode != ESequencerKeyMode::ManualKey && PropertyChangedParams.KeyMode != ESequencerKeyMode::ManualKeyForced)
+	{
+		return;
+	}
+
+	UProperty* Property = PropertyChangedParams.PropertyPath.GetLeafMostProperty().Property.Get();
+	if (!Property)
+	{
+		return;
+	}
+
+	MovieSceneSequence->SetFlags(RF_Transactional);
+	FScopedTransaction Transaction( LOCTEXT("AddMediaPropertyTrack", "Add Media Property Track"));
+
+	FName UniqueName(*PropertyChangedParams.PropertyPath.ToString(TEXT(".")));
+
+	for (UObject* Object : PropertyChangedParams.ObjectsThatChanged)
+	{
+		FFindOrCreateHandleResult HandleResult = FindOrCreateHandleToObject(Object);
+		if (HandleResult.Handle.IsValid())
+		{
+			// Ensure it has a track
+			FFindOrCreateTrackResult TrackResult = FindOrCreateTrackForObject(HandleResult.Handle, UMovieSceneMediaPlayerPropertyTrack::StaticClass(), Property->GetFName());
+			if (TrackResult.bWasCreated)
+			{
+				InitializeNewTrack(CastChecked<UMovieSceneMediaPlayerPropertyTrack>(TrackResult.Track), PropertyChangedParams);
+
+				GetSequencer()->NotifyMovieSceneDataChanged(EMovieSceneDataChangeType::MovieSceneStructureItemAdded);
+			}
+		}
+	}
+}
+
+UMovieSceneTrack* FMediaPlayerPropertyTrackEditor::AddTrack(UMovieScene* FocusedMovieScene, const FGuid& ObjectHandle, TSubclassOf<class UMovieSceneTrack> TrackClass, FName UniqueTypeName)
+{
+	UMovieSceneTrack* NewTrack = FPropertyTrackEditor::AddTrack(FocusedMovieScene, ObjectHandle, TrackClass, UniqueTypeName);
+	UMovieSceneMediaPlayerPropertyTrack* MediaPlayerTrack = Cast<UMovieSceneMediaPlayerPropertyTrack>(NewTrack);
+	return NewTrack;
+}
+
+TSharedPtr<SWidget> FMediaPlayerPropertyTrackEditor::BuildOutlinerEditWidget(const FGuid& ObjectBinding, UMovieSceneTrack* Track, const FBuildEditWidgetParams& Params)
+{
+	UMovieSceneMediaPlayerPropertyTrack* MediaTrack = Cast<UMovieSceneMediaPlayerPropertyTrack>(Track);
+	if (!MediaTrack)
+	{
+		return SNullWidget::NullWidget;
+	}
+
+	auto CreatePicker = [this, MediaTrack]
+	{
+		FAssetPickerConfig AssetPickerConfig;
+		{
+			AssetPickerConfig.OnAssetSelected     = FOnAssetSelected::CreateRaw(this,     &FMediaPlayerPropertyTrackEditor::AddNewSection,             MediaTrack);
+			AssetPickerConfig.OnAssetEnterPressed = FOnAssetEnterPressed::CreateRaw(this, &FMediaPlayerPropertyTrackEditor::AddNewSectionEnterPressed, MediaTrack);
+			AssetPickerConfig.bAllowNullSelection = false;
+			AssetPickerConfig.InitialAssetViewType = EAssetViewType::List;
+			AssetPickerConfig.Filter.bRecursiveClasses = true;
+			AssetPickerConfig.Filter.ClassNames.Add(UMediaSource::StaticClass()->GetFName());
+		}
+
+		FContentBrowserModule& ContentBrowserModule = FModuleManager::Get().LoadModuleChecked<FContentBrowserModule>(TEXT("ContentBrowser"));
+
+		TSharedRef<SBox> Picker = SNew(SBox)
+			.WidthOverride(500.0f)
+			.HeightOverride(400.f)
+			[
+				ContentBrowserModule.Get().CreateAssetPicker(AssetPickerConfig)
+			];
+
+		return Picker;
+	};
+
+	return SNew(SHorizontalBox)
+
+		+ SHorizontalBox::Slot()
+		.AutoWidth()
+		.VAlign(VAlign_Center)
+		[
+			FSequencerUtilities::MakeAddButton(LOCTEXT("AddMediaSection_Text", "Media"), FOnGetContent::CreateLambda(CreatePicker), Params.NodeIsHovered, GetSequencer())
+		];
+}
+
+TSharedRef<ISequencerSection> FMediaPlayerPropertyTrackEditor::MakeSectionInterface(UMovieSceneSection& SectionObject, UMovieSceneTrack& Track, FGuid ObjectBinding)
+{
+	return MakeShared<FMediaPlayerPropertySection>(CastChecked<UMovieSceneMediaPlayerPropertySection>(&SectionObject));
+}
+
+bool FMediaPlayerPropertyTrackEditor::SupportsSequence(UMovieSceneSequence* InSequence) const
+{
+	return (InSequence != nullptr) && (InSequence->GetClass()->GetName() == TEXT("LevelSequence"));
+}
+
+const FSlateBrush* FMediaPlayerPropertyTrackEditor::GetIconBrush() const
+{
+	return FEditorStyle::GetBrush("Sequencer.Tracks.Media");
+}
+
+void FMediaPlayerPropertyTrackEditor::AddNewSection(const FAssetData& AssetData, UMovieSceneMediaPlayerPropertyTrack* MediaTrack)
+{
+	FSlateApplication::Get().DismissAllMenus();
+
+	UObject* SelectedObject = AssetData.GetAsset();
+	if (SelectedObject)
+	{
+		UMediaSource* MediaSource = CastChecked<UMediaSource>(AssetData.GetAsset());
+
+		const FScopedTransaction Transaction(LOCTEXT("AddMedia_Transaction", "Add Media"));
+
+		MediaTrack->Modify();
+
+		TArray<UMovieSceneSection*> Sections = MediaTrack->GetAllSections();
+
+		const float DefaultMediaSectionDuration = 5.0f;
+		FFrameRate TickResolution = MediaTrack->GetTypedOuter<UMovieScene>()->GetTickResolution();
+		FFrameTime DurationToUse  = DefaultMediaSectionDuration * TickResolution;
+
+		// add the section
+		UMovieSceneMediaPlayerPropertySection* NewSection = CastChecked<UMovieSceneMediaPlayerPropertySection>(MediaTrack->CreateNewSection());
+
+		NewSection->InitialPlacement(MediaTrack->GetAllSections(), GetSequencer()->GetLocalTime().Time.FrameNumber, DurationToUse.FrameNumber.Value, false);
+		NewSection->MediaSource = MediaSource;
+
+		MediaTrack->AddSection(*NewSection);
+
+		GetSequencer()->EmptySelection();
+		GetSequencer()->SelectSection(NewSection);
+		GetSequencer()->ThrobSectionSelection();
+
+		GetSequencer()->NotifyMovieSceneDataChanged(EMovieSceneDataChangeType::MovieSceneStructureItemAdded);
+	}
+}
+
+void FMediaPlayerPropertyTrackEditor::AddNewSectionEnterPressed(const TArray<FAssetData>& AssetData, UMovieSceneMediaPlayerPropertyTrack* Track)
+{
+	if (AssetData.Num() > 0)
+	{
+		AddNewSection(AssetData[0].GetAsset(), Track);
+	}
+}
+
+#undef LOCTEXT_NAMESPACE
diff --git a/Engine/Plugins/Media/MediaCompositing/Source/MediaCompositingEditor/Private/Sequencer/MediaPlayerPropertyTrackEditor.h b/Engine/Plugins/Media/MediaCompositing/Source/MediaCompositingEditor/Private/Sequencer/MediaPlayerPropertyTrackEditor.h
new file mode 100644
--- /dev/null
+++ b/Engine/Plugins/Media/MediaCompositing/Source/MediaCompositingEditor/Private/Sequencer/MediaPlayerPropertyTrackEditor.h
@@ -0,0 +1,50 @@
+// Copyright 1998-2020 Epic Games, Inc. All Rights Reserved.
+
+#pragma once
+
+#include "PropertyTrackEditor.h"
+#include "MovieSceneMediaPlayerPropertyTrack.h"
+#include "MediaPlayer.h"
+
+struct FAssetData;
+class ISequencer;
+class UMediaSource;
+class UMovieSceneMediaPlayerPropertyTrack;
+
+/**
+ * Track editor that understands how to animate MediaPlayer properties on objects
+ */
+class FMediaPlayerPropertyTrackEditor
+	: public FPropertyTrackEditor<UMovieSceneMediaPlayerPropertyTrack>
+{
+public:
+
+	FMediaPlayerPropertyTrackEditor(TSharedRef<ISequencer> InSequencer)
+		: FPropertyTrackEditor(InSequencer, GetAnimatedPropertyTypes())
+	{}
+
+	static TSharedRef<ISequencerTrackEditor> CreateTrackEditor(TSharedRef<ISequencer> OwningSequencer)
+	{
+		return MakeShared<FMediaPlayerPropertyTrackEditor>(OwningSequencer);
+	}
+
+
+	static TArray<FAnimatedPropertyKey, TInlineAllocator<1>> GetAnimatedPropertyTypes()
+	{
+		return TArray<FAnimatedPropertyKey, TInlineAllocator<1>>({ FAnimatedPropertyKey::FromObjectType(UMediaPlayer::StaticClass()) });
+	}
+
+private:
+
+	virtual UMovieSceneTrack* AddTrack(UMovieScene* FocusedMovieScene, const FGuid& ObjectHandle, TSubclassOf<class UMovieSceneTrack> TrackClass, FName UniqueTypeName) override;
+	virtual TSharedPtr<SWidget> BuildOutlinerEditWidget(const FGuid& ObjectBinding, UMovieSceneTrack* Track, const FBuildEditWidgetParams& Params) override;
+	virtual bool SupportsSequence(UMovieSceneSequence* InSequence) const override;
+	virtual const FSlateBrush* GetIconBrush() const override;
+	virtual TSharedRef<ISequencerSection> MakeSectionInterface(UMovieSceneSection& SectionObject, UMovieSceneTrack& Track, FGuid ObjectBinding) override;
+	virtual void OnAnimatedPropertyChanged(const FPropertyChangedParams& PropertyChangedParams) override;
+
+	virtual void GenerateKeysFromPropertyChanged(const FPropertyChangedParams&, FGeneratedTrackKeys&) override {}
+
+	void AddNewSection(const FAssetData& AssetData, UMovieSceneMediaPlayerPropertyTrack* MediaTrack);
+	void AddNewSectionEnterPressed(const TArray<FAssetData>& AssetData, UMovieSceneMediaPlayerPropertyTrack* Track);
+};
diff --git a/Engine/Source/Editor/Sequencer/Private/SSequencerPlayRateCombo.cpp b/Engine/Source/Editor/Sequencer/Private/SSequencerPlayRateCombo.cpp
--- a/Engine/Source/Editor/Sequencer/Private/SSequencerPlayRateCombo.cpp
+++ b/Engine/Source/Editor/Sequencer/Private/SSequencerPlayRateCombo.cpp
@@ -15,6 +15,11 @@
 #include "ScopedTransaction.h"
 #include "Widgets/SFrameRateEntryBox.h"
 #include "EditorFontGlyphs.h"
+#include "Evaluation/IMovieSceneCustomClockSource.h"
+#include "SceneOutlinerModule.h"
+#include "SceneOutlinerPublicTypes.h"
+#include "IContentBrowserSingleton.h"
+#include "ContentBrowserModule.h"
 #include "Misc/Timecode.h"
 
 #define LOCTEXT_NAMESPACE "SSequencerPlayRateCombo"
@@ -343,23 +348,109 @@ void SSequencerPlayRateCombo::PopulateClockSourceMenu(FMenuBuilder& MenuBuilder)
 			{
 				EUpdateClockSource Value = (EUpdateClockSource)ClockSourceEnum->GetValueByIndex(Index);
 
-				MenuBuilder.AddMenuEntry(
-					ClockSourceEnum->GetDisplayNameTextByIndex(Index),
-					ClockSourceEnum->GetToolTipTextByIndex(Index),
-					FSlateIcon(),
-					FUIAction(
-						FExecuteAction::CreateSP(this, &SSequencerPlayRateCombo::SetClockSource, Value),
-						FCanExecuteAction(),
-						FIsActionChecked::CreateLambda([=]{ return RootSequence->GetMovieScene()->GetClockSource() == Value; })
-					),
-					NAME_None,
-					EUserInterfaceActionType::RadioButton
-				);
+				if (Value == EUpdateClockSource::Custom)
+				{
+					MenuBuilder.AddSubMenu(
+						ClockSourceEnum->GetDisplayNameTextByIndex(Index),
+						ClockSourceEnum->GetToolTipTextByIndex(Index),
+						FNewMenuDelegate::CreateSP(this, &SSequencerPlayRateCombo::PopulateCustomClockSourceMenu),
+						FUIAction(
+							FExecuteAction::CreateSP(this, &SSequencerPlayRateCombo::SetClockSource, Value),
+							FCanExecuteAction(),
+							FIsActionChecked::CreateLambda([=]{ return RootSequence->GetMovieScene()->GetClockSource() == Value; })
+						),
+						NAME_None,
+						EUserInterfaceActionType::RadioButton
+					);
+				}
+				else
+				{
+					MenuBuilder.AddMenuEntry(
+						ClockSourceEnum->GetDisplayNameTextByIndex(Index),
+						ClockSourceEnum->GetToolTipTextByIndex(Index),
+						FSlateIcon(),
+						FUIAction(
+							FExecuteAction::CreateSP(this, &SSequencerPlayRateCombo::SetClockSource, Value),
+							FCanExecuteAction(),
+							FIsActionChecked::CreateLambda([=]{ return RootSequence->GetMovieScene()->GetClockSource() == Value; })
+						),
+						NAME_None,
+						EUserInterfaceActionType::RadioButton
+					);
+				}
 			}
 		}
 	}
 }
 
+void SSequencerPlayRateCombo::PopulateCustomClockSourceMenu(FMenuBuilder& MenuBuilder)
+{
+	auto ActorClockSourceMenu = [this](FMenuBuilder& ActorMenuBuilder)
+	{
+		auto IsActorValid = [](const AActor* Actor)
+		{
+			return Actor && Actor->GetClass()->ImplementsInterface(UMovieSceneCustomClockSource::StaticClass());
+		};
+
+		using namespace SceneOutliner;
+
+		// Set up a menu entry to assign an actor to the object binding node
+		FInitializationOptions InitOptions;
+		InitOptions.Mode = ESceneOutlinerMode::ActorPicker;
+
+		// We hide the header row to keep the UI compact.
+		InitOptions.bShowHeaderRow = false;
+		InitOptions.bShowSearchBox = true;
+		InitOptions.bShowCreateNewFolder = false;
+		InitOptions.bFocusSearchBoxWhenOpened = true;
+		// Only want the actor label column
+		InitOptions.ColumnMap.Add(FBuiltInColumnTypes::Label(), FColumnInfo(EColumnVisibility::Visible, 0));
+
+		// Only display actors that are not possessed already
+		InitOptions.Filters->AddFilterPredicate( FActorFilterPredicate::CreateLambda( IsActorValid ) );
+
+		// actor selector to allow the user to choose an actor
+		FSceneOutlinerModule& SceneOutlinerModule = FModuleManager::LoadModuleChecked<FSceneOutlinerModule>("SceneOutliner");
+		ActorMenuBuilder.AddWidget(
+			SNew(SBox)
+			.MaxDesiredHeight(400.0f)
+			.WidthOverride(300.0f)
+			[
+				SceneOutlinerModule.CreateSceneOutliner(InitOptions, FOnActorPicked::CreateLambda([this](AActor* In){ this->SetCustomClockSource(In); }))
+			],
+			FText(),
+			true /*bNoIndent*/
+		);
+	};
+
+	auto AssetClockSourceMenu = [this](FMenuBuilder& ActorMenuBuilder)
+	{
+		FAssetPickerConfig AssetPickerConfig;
+		AssetPickerConfig.Filter.ClassNames.Add(UMovieSceneCustomClockSource::StaticClass()->GetFName());
+		AssetPickerConfig.bAllowNullSelection = false;
+		AssetPickerConfig.Filter.bRecursiveClasses = true;
+		AssetPickerConfig.OnAssetSelected = FOnAssetSelected::CreateLambda([this](const FAssetData& In){ this->SetCustomClockSource(In.GetAsset()); });
+		AssetPickerConfig.InitialAssetViewType = EAssetViewType::List;
+		AssetPickerConfig.bAllowDragging = false;
+
+		FContentBrowserModule& ContentBrowserModule = FModuleManager::Get().LoadModuleChecked<FContentBrowserModule>(TEXT("ContentBrowser"));
+
+		ActorMenuBuilder.AddWidget(
+			SNew(SBox)
+			.WidthOverride(300)
+			.HeightOverride(300)
+			[
+				ContentBrowserModule.Get().CreateAssetPicker(AssetPickerConfig)
+			],
+			FText(),
+			true /*bNoIndent*/
+		);
+	};
+
+	MenuBuilder.AddSubMenu(LOCTEXT("ActorClockSource", "Actors"), FText(), FNewMenuDelegate::CreateLambda(ActorClockSourceMenu));
+	MenuBuilder.AddSubMenu(LOCTEXT("AssetClockSource", "Assets"), FText(), FNewMenuDelegate::CreateLambda(AssetClockSourceMenu));
+}
+
 void SSequencerPlayRateCombo::AddMenuEntry(FMenuBuilder& MenuBuilder, const FCommonFrameRateInfo& Info)
 {
 	MenuBuilder.AddMenuEntry(
@@ -398,6 +489,28 @@ void SSequencerPlayRateCombo::SetClockSource(EUpdateClockSource NewClockSource)
 	}
 }
 
+void SSequencerPlayRateCombo::SetCustomClockSource(UObject* InClockSource)
+{
+	TSharedPtr<FSequencer> Sequencer    = WeakSequencer.Pin();
+	UMovieSceneSequence*   RootSequence = Sequencer.IsValid() ? Sequencer->GetRootMovieSceneSequence() : nullptr;
+	if (RootSequence && InClockSource)
+	{
+		UMovieScene* MovieScene = RootSequence->GetMovieScene();
+
+		if (MovieScene->IsReadOnly())
+		{
+			return;
+		}
+
+		FScopedTransaction ScopedTransaction(LOCTEXT("SetClockSource", "Set Clock Source"));
+
+		MovieScene->Modify();
+		MovieScene->SetClockSource(InClockSource);
+
+		Sequencer->ResetTimeController();
+	}
+}
+
 void SSequencerPlayRateCombo::SetDisplayRate(FFrameRate InFrameRate)
 {
 	TSharedPtr<FSequencer> Sequencer = WeakSequencer.Pin();
diff --git a/Engine/Source/Editor/Sequencer/Private/SSequencerPlayRateCombo.h b/Engine/Source/Editor/Sequencer/Private/SSequencerPlayRateCombo.h
--- a/Engine/Source/Editor/Sequencer/Private/SSequencerPlayRateCombo.h
+++ b/Engine/Source/Editor/Sequencer/Private/SSequencerPlayRateCombo.h
@@ -45,8 +45,11 @@ private:
 	TSharedRef<SWidget> OnCreateMenu();
 	void PopulateIncompatibleRatesMenu(FMenuBuilder& MenuBuilder);
 	void PopulateClockSourceMenu(FMenuBuilder& MenuBuilder);
+	void PopulateCustomClockSourceMenu(FMenuBuilder& MenuBuilder);
 	void AddMenuEntry(FMenuBuilder& MenuBuilder, const FCommonFrameRateInfo& Info);
 
+	void SetCustomClockSource(UObject* Object);
+
 	void OnToggleFrameLocked();
 	ECheckBoxState OnGetFrameLockedCheckState() const;
 
diff --git a/Engine/Source/Editor/Sequencer/Private/Sequencer.cpp b/Engine/Source/Editor/Sequencer/Private/Sequencer.cpp
--- a/Engine/Source/Editor/Sequencer/Private/Sequencer.cpp
+++ b/Engine/Source/Editor/Sequencer/Private/Sequencer.cpp
@@ -639,7 +639,6 @@ void FSequencer::Tick(float InDeltaTime)
 	TimeController->Tick(InDeltaTime, PlaybackSpeed * Dilation);
 
 	FQualifiedFrameTime GlobalTime = GetGlobalTime();
-	FFrameTime NewGlobalTime = TimeController->RequestCurrentTime(GlobalTime, PlaybackSpeed * Dilation);
 
 	static const float AutoScrollFactor = 0.1f;
 
@@ -665,8 +664,10 @@ void FSequencer::Tick(float InDeltaTime)
 
 	if (PlaybackState == EMovieScenePlayerStatus::Playing)
 	{
+		FFrameTime NewGlobalTime = TimeController->RequestCurrentTime(GlobalTime, PlaybackSpeed * Dilation);
+
 		// Put the time into local space
-		SetLocalTimeLooped(NewGlobalTime * RootToLocalTransform);
+ 		SetLocalTimeLooped(NewGlobalTime * RootToLocalTransform);
 
 		if (IsAutoScrollEnabled() && GetPlaybackStatus() == EMovieScenePlayerStatus::Playing)
 		{
@@ -10882,12 +10883,19 @@ void FSequencer::UpdateTimeBases()
 
 void FSequencer::ResetTimeController()
 {
-	switch (GetRootMovieSceneSequence()->GetMovieScene()->GetClockSource())
+	UMovieScene* MovieScene = GetRootMovieSceneSequence()->GetMovieScene();
+	switch (MovieScene->GetClockSource())
+	{
+	case EUpdateClockSource::Audio:    TimeController = MakeShared<FMovieSceneTimeController_AudioClock>();         break;
+	case EUpdateClockSource::Platform: TimeController = MakeShared<FMovieSceneTimeController_PlatformClock>();      break;
+	case EUpdateClockSource::Timecode: TimeController = MakeShared<FMovieSceneTimeController_TimecodeClock>();      break;
+	case EUpdateClockSource::Custom:   TimeController = MovieScene->MakeCustomTimeController(GetPlaybackContext()); break;
+	default:                           TimeController = MakeShared<FMovieSceneTimeController_Tick>();               break;
+	}
+
+	if (!TimeController)
 	{
-	case EUpdateClockSource::Audio:    TimeController = MakeShared<FMovieSceneTimeController_AudioClock>();    break;
-	case EUpdateClockSource::Platform: TimeController = MakeShared<FMovieSceneTimeController_PlatformClock>(); break;
-	case EUpdateClockSource::Timecode: TimeController = MakeShared<FMovieSceneTimeController_TimecodeClock>(); break;
-	default:                           TimeController = MakeShared<FMovieSceneTimeController_Tick>();          break;
+		TimeController = MakeShared<FMovieSceneTimeController_Tick>();
 	}
 
 	TimeController->PlayerStatusChanged(PlaybackState, GetGlobalTime());
diff --git a/Engine/Source/Runtime/LevelSequence/LevelSequence.Build.cs b/Engine/Source/Runtime/LevelSequence/LevelSequence.Build.cs
--- a/Engine/Source/Runtime/LevelSequence/LevelSequence.Build.cs
+++ b/Engine/Source/Runtime/LevelSequence/LevelSequence.Build.cs
@@ -20,6 +20,12 @@ public class LevelSequence : ModuleRules
 			}
 		);
 
+		PrivateDependencyModuleNames.AddRange(
+			new string[] {
+				"MediaAssets",
+			}
+		);
+
 		if (Target.bBuildEditor == true)
 		{
 			PrivateDependencyModuleNames.AddRange(
diff --git a/Engine/Source/Runtime/LevelSequence/Private/SequenceMediaController.cpp b/Engine/Source/Runtime/LevelSequence/Private/SequenceMediaController.cpp
new file mode 100644
--- /dev/null
+++ b/Engine/Source/Runtime/LevelSequence/Private/SequenceMediaController.cpp
@@ -0,0 +1,162 @@
+// Copyright 1998-2020 Epic Games, Inc. All Rights Reserved.
+
+#include "SequenceMediaController.h"
+
+#include "Engine/World.h"
+#include "GameFramework/GameState.h"
+#include "Net/UnrealNetwork.h"
+#include "TimerManager.h"
+
+#include "IMediaEventSink.h"
+#include "MediaComponent.h"
+#include "MediaPlayer.h"
+#include "MediaPlaylist.h"
+#include "LevelSequenceActor.h"
+#include "LevelSequenceModule.h"
+
+
+ALevelSequenceMediaController::ALevelSequenceMediaController(const FObjectInitializer& Init)
+	: Super(Init)
+{
+	bReplicates = true;
+	ServerStartTimeSeconds = -MIN_flt;
+	SequencePositionSeconds = 0.0;
+
+	MediaComponent = Init.CreateDefaultSubobject<UMediaComponent>(this, "MediaComponent");
+	AddOwnedComponent(MediaComponent);
+}
+
+void ALevelSequenceMediaController::GetLifetimeReplicatedProps(TArray<FLifetimeProperty>& OutLifetimeProps) const
+{
+	Super::GetLifetimeReplicatedProps(OutLifetimeProps);
+
+	DOREPLIFETIME(ALevelSequenceMediaController, ServerStartTimeSeconds);
+}
+
+void ALevelSequenceMediaController::OnRep_ServerStartTimeSeconds()
+{
+	if (ServerStartTimeSeconds != -MIN_flt)
+	{
+		Client_Play();
+	}
+}
+
+void ALevelSequenceMediaController::Play()
+{
+	if (HasAuthority())
+	{
+		UWorld*         World     = GetWorld();
+		AGameStateBase* GameState = World ? World->GetGameState() : nullptr;
+
+		if (GameState)
+		{
+			ServerStartTimeSeconds = GameState->GetServerWorldTimeSeconds();
+		}
+
+		if (GetNetMode() != NM_DedicatedServer)
+		{
+			OnRep_ServerStartTimeSeconds();
+		}
+	}
+	else
+	{
+		FFrame::KismetExecutionMessage(*FString::Printf(TEXT("Cannot initiate playback when Sequence does not have authority (%s)"), *GetName()), ELogVerbosity::Warning);
+	}
+}
+
+void ALevelSequenceMediaController::SynchronizeToServer(float DesyncThresholdSeconds )
+{
+	if (HasAuthority())
+	{
+		FFrame::KismetExecutionMessage(*FString::Printf(TEXT("Cannot synchronize when Sequence has authority (%s)"), *GetName()), ELogVerbosity::Warning);
+	}
+	else
+	{
+		Client_ConditionallyForceTime(DesyncThresholdSeconds);
+	}
+}
+
+void ALevelSequenceMediaController::Client_Play()
+{
+	check(ServerStartTimeSeconds != -MIN_flt);
+
+	if (!Sequence)
+	{
+		return;
+	}
+
+	float CurrentServerTime = 0.f;
+
+	UWorld*         World     = GetWorld();
+	AGameStateBase* GameState = World ? World->GetGameState() : nullptr;
+	if (GameState)
+	{
+		CurrentServerTime = GameState->GetServerWorldTimeSeconds() - ServerStartTimeSeconds;
+	}
+
+	const float StartTimeSeconds = Sequence->GetSequencePlayer()->GetStartTime().AsSeconds() + CurrentServerTime;
+
+	UE_LOG(LogLevelSequence, Log, TEXT("Initiating playback of sequence '%s' starting at time %fs"), *Sequence->GetName(), StartTimeSeconds);
+
+	Sequence->GetSequencePlayer()->JumpToSeconds(StartTimeSeconds);
+	Sequence->GetSequencePlayer()->Play();
+}
+
+void ALevelSequenceMediaController::Client_ConditionallyForceTime(float DesyncThresholdSeconds)
+{
+	UWorld*         World     = GetWorld();
+	AGameStateBase* GameState = World ? World->GetGameState() : nullptr;
+
+	if (!GameState || !Sequence)
+	{
+		return;
+	}
+
+	if (!Sequence->GetSequencePlayer()->IsPlaying())
+	{
+		return;
+	}
+
+	const double SequenceOffset = Sequence->GetSequencePlayer()->GetCurrentTime().AsSeconds() - Sequence->GetSequencePlayer()->GetStartTime().AsSeconds();
+
+	const float ExpectedServerTime = ServerStartTimeSeconds + static_cast<float>(SequenceOffset);
+	const float CurrentServerTime = GameState->GetServerWorldTimeSeconds();
+	
+	const float Difference = FMath::Abs(CurrentServerTime - ExpectedServerTime);
+	if (Difference > DesyncThresholdSeconds)
+	{
+		const double NewTimeSeconds = Sequence->GetSequencePlayer()->GetStartTime().AsSeconds() + (CurrentServerTime - ServerStartTimeSeconds);
+		Sequence->GetSequencePlayer()->JumpToSeconds(NewTimeSeconds);
+
+		UE_LOG(LogLevelSequence, Warning, TEXT("Forcibly synchronizing sequence '%s' to time %f to server time (it is out by %fs)."), *Sequence->GetName(), NewTimeSeconds, Difference);
+	}
+}
+
+void ALevelSequenceMediaController::OnStartPlaying(const FQualifiedFrameTime& InStartTime)
+{
+	SequencePositionSeconds = InStartTime.AsSeconds();
+}
+
+void ALevelSequenceMediaController::OnTick(float DeltaSeconds, float InPlayRate)
+{
+	double StartOffset = 0.0;
+
+	UMediaPlayer* MediaPlayer = MediaComponent->GetMediaPlayer();
+	if (MediaPlayer->GetRate() != 0.f || MediaPlayer->IsPlaying())
+	{
+		SequencePositionSeconds = MediaPlayer->GetTime().GetTotalSeconds();
+	}
+	else
+	{
+		SequencePositionSeconds += DeltaSeconds * InPlayRate;
+	}
+}
+
+void ALevelSequenceMediaController::OnStopPlaying(const FQualifiedFrameTime& InStopTime)
+{
+}
+
+FFrameTime ALevelSequenceMediaController::OnRequestCurrentTime(const FQualifiedFrameTime& InCurrentTime, float InPlayRate)
+{
+	return SequencePositionSeconds * InCurrentTime.Rate;
+}
\ No newline at end of file
diff --git a/Engine/Source/Runtime/LevelSequence/Public/SequenceMediaController.h b/Engine/Source/Runtime/LevelSequence/Public/SequenceMediaController.h
new file mode 100644
--- /dev/null
+++ b/Engine/Source/Runtime/LevelSequence/Public/SequenceMediaController.h
@@ -0,0 +1,94 @@
+// Copyright 1998-2020 Epic Games, Inc. All Rights Reserved.
+
+#pragma once
+
+#include "CoreTypes.h"
+#include "UObject/ObjectMacros.h"
+#include "GameFramework/Actor.h"
+#include "Evaluation/IMovieSceneCustomClockSource.h"
+#include "LevelSequenceActor.h"
+#include "SequenceMediaController.generated.h"
+
+enum class EMediaEvent;
+
+class FDelegateHandle;
+class UMediaPlayer;
+
+/**
+ * Replicated actor class that is responsible for instigating various cinematic assets (Media, Audio, Level Sequences) in a synchronized fasion
+ */
+UCLASS(hideCategories=(Rendering, Physics, LOD, Activation, Input, Collision))
+class LEVELSEQUENCE_API ALevelSequenceMediaController
+	: public AActor
+	, public IMovieSceneCustomClockSource
+{
+public:
+
+	GENERATED_BODY()
+
+	ALevelSequenceMediaController(const FObjectInitializer& Init);
+
+	UFUNCTION(BlueprintCallable, Category="Synchronization")
+	void Play();
+
+	/**
+	 * Access this actor's media component
+	 */
+	UFUNCTION(BlueprintPure, Category="Synchronization")
+	UMediaComponent* GetMediaComponent() const
+	{
+		return MediaComponent;
+	}
+
+	/**
+	 * Access this actor's Level Sequence Actor
+	 */
+	UFUNCTION(BlueprintPure, Category="Synchronization")
+	ALevelSequenceActor* GetSequence() const
+	{
+		return Sequence;
+	}
+
+	/**
+	 * Forcibly synchronize the sequence to the server's position if it has diverged by more than the specified threshold
+	 */
+	UFUNCTION(BlueprintCallable, Category="Synchronization")
+	void SynchronizeToServer(float DesyncThresholdSeconds = 2.f);
+
+private:
+
+	void Client_Play();
+
+	void Client_ConditionallyForceTime(float DesyncThresholdSeconds);
+
+private:
+
+	virtual void GetLifetimeReplicatedProps(TArray<FLifetimeProperty>& OutLifetimeProps) const;
+
+private:
+
+	virtual void OnTick(float DeltaSeconds, float InPlayRate) override final;
+	virtual void OnStartPlaying(const FQualifiedFrameTime& InStartTime) override final;
+	virtual void OnStopPlaying(const FQualifiedFrameTime& InStopTime) override final;
+	virtual FFrameTime OnRequestCurrentTime(const FQualifiedFrameTime& InCurrentTime, float InPlayRate) override final;
+
+private:
+
+	UFUNCTION()
+	void OnRep_ServerStartTimeSeconds();
+
+	/** Pointer to the sequence actor to use for playback */
+	UPROPERTY(EditAnywhere, BlueprintGetter=GetSequence, Category="Synchronization")
+	ALevelSequenceActor* Sequence;
+
+	/** Media component that contains the media player to synchronize with */
+	UPROPERTY(VisibleAnywhere, BlueprintGetter=GetMediaComponent, Category="Synchronization")
+	UMediaComponent* MediaComponent;
+
+	/** Replicated time at which the server started the sequence (taken from AGameStateBase::GetServerWorldTimeSeconds) */
+	UPROPERTY(BlueprintReadOnly, Category="Synchronization", replicated, ReplicatedUsing=OnRep_ServerStartTimeSeconds, meta=(AllowPrivateAccess="true"))
+	float ServerStartTimeSeconds;
+
+	/** Time to use for the sequence playback position */
+	double SequencePositionSeconds;
+};
diff --git a/Engine/Source/Runtime/MediaAssets/Private/Assets/MediaComponent.cpp b/Engine/Source/Runtime/MediaAssets/Private/Assets/MediaComponent.cpp
new file mode 100644
--- /dev/null
+++ b/Engine/Source/Runtime/MediaAssets/Private/Assets/MediaComponent.cpp
@@ -0,0 +1,48 @@
+// Copyright 1998-2020 Epic Games, Inc. All Rights Reserved.
+
+#include "MediaComponent.h"
+#include "MediaPlayer.h"
+#include "MediaTexture.h"
+#include "IMediaEventSink.h"
+
+UMediaComponent::UMediaComponent(const FObjectInitializer& ObjectInitializer)
+	: Super(ObjectInitializer)
+{
+	PrimaryComponentTick.bCanEverTick = true;
+	bAutoActivate = true;
+
+	if (!HasAnyFlags(RF_ClassDefaultObject))
+	{
+		MediaPlayer = NewObject<UMediaPlayer>(this, "MediaPlayer", RF_Transient);
+		MediaPlayer->SetLooping(false);
+		MediaPlayer->PlayOnOpen = false;
+
+		MediaTexture = NewObject<UMediaTexture>(this, "MediaTexture", RF_Transient);
+	}
+}
+
+
+UMediaPlayer* UMediaComponent::GetMediaPlayer() const
+{
+	return MediaPlayer;
+}
+
+
+UMediaTexture* UMediaComponent::GetMediaTexture() const
+{
+	return MediaTexture;
+}
+
+void UMediaComponent::OnRegister()
+{
+	Super::OnRegister();
+
+	if (MediaPlayer)
+	{
+		if (MediaTexture)
+		{
+			MediaTexture->SetMediaPlayer(MediaPlayer);
+			MediaTexture->UpdateResource();
+		}
+	}
+}
diff --git a/Engine/Source/Runtime/MediaAssets/Public/MediaComponent.h b/Engine/Source/Runtime/MediaAssets/Public/MediaComponent.h
new file mode 100644
--- /dev/null
+++ b/Engine/Source/Runtime/MediaAssets/Public/MediaComponent.h
@@ -0,0 +1,41 @@
+// Copyright 1998-2020 Epic Games, Inc. All Rights Reserved.
+
+#pragma once
+
+#include "Components/ActorComponent.h"
+
+#include "MediaComponent.generated.h"
+
+class UMediaPlayer;
+class UMediaTexture;
+
+UCLASS(Category="Media", hideCategories=(Media, Activation))
+class MEDIAASSETS_API UMediaComponent : public UActorComponent
+{
+public:
+	GENERATED_BODY()
+
+	UMediaComponent(const FObjectInitializer& ObjectInitializer);
+
+	/** Get the media player that this component owns */
+	UFUNCTION(BlueprintPure, Category="Media")
+	UMediaPlayer* GetMediaPlayer() const;
+
+	/** Get the media texture that this component owns, bound to its media player. */
+	UFUNCTION(BlueprintPure, Category="Media")
+	UMediaTexture* GetMediaTexture() const;
+
+private:
+
+	virtual void OnRegister() override;
+
+private:
+
+	/** This component's media texture */
+	UPROPERTY(Category="Media", Instanced, Transient, BlueprintGetter=GetMediaTexture)
+	UMediaTexture* MediaTexture;
+
+	/** This component's media player */
+	UPROPERTY(Category="Media", Instanced, Transient, BlueprintGetter=GetMediaPlayer, Interp)
+	UMediaPlayer* MediaPlayer;
+};
\ No newline at end of file
diff --git a/Engine/Source/Runtime/MovieScene/Private/Evaluation/IMovieSceneCustomClockSource.cpp b/Engine/Source/Runtime/MovieScene/Private/Evaluation/IMovieSceneCustomClockSource.cpp
new file mode 100644
--- /dev/null
+++ b/Engine/Source/Runtime/MovieScene/Private/Evaluation/IMovieSceneCustomClockSource.cpp
@@ -0,0 +1,95 @@
+// Copyright 1998-2020 Epic Games, Inc. All Rights Reserved.
+
+#include "Evaluation/IMovieSceneCustomClockSource.h"
+#include "UObject/Package.h"
+
+FMovieSceneTimeController_Custom::FMovieSceneTimeController_Custom(const FSoftObjectPath& InObjectPath, TWeakObjectPtr<> InWeakPlaybackContext)
+	: WeakPlaybackContext(InWeakPlaybackContext)
+	, InterfacePtr(nullptr)
+	, ObjectPath(InObjectPath)
+{
+#if WITH_EDITOR
+	UObject* PlaybackContext = WeakPlaybackContext.Get();
+
+	UPackage* Package = PlaybackContext ? PlaybackContext->GetOutermost() : nullptr;
+	if (Package && Package->PIEInstanceID != INDEX_NONE)
+	{
+		if (!ObjectPath.FixupForPIE(Package->PIEInstanceID))
+		{
+			// log error?
+		}
+	}
+#endif
+
+	ResolveInterfacePtr();
+}
+
+void FMovieSceneTimeController_Custom::OnTick(float DeltaSeconds, float InPlayRate)
+{
+	if (WeakObject.IsStale())
+	{
+		ResolveInterfacePtr();
+	}
+
+	if (WeakObject.IsValid())
+	{
+		InterfacePtr->OnTick(DeltaSeconds, InPlayRate);
+	}
+}
+
+void FMovieSceneTimeController_Custom::OnStartPlaying(const FQualifiedFrameTime& InStartTime)
+{
+	if (WeakObject.IsStale())
+	{
+		ResolveInterfacePtr();
+	}
+
+	if (WeakObject.IsValid())
+	{
+		InterfacePtr->OnStartPlaying(InStartTime);
+	}
+}
+
+void FMovieSceneTimeController_Custom::OnStopPlaying(const FQualifiedFrameTime& InStopTime)
+{
+	if (WeakObject.IsStale())
+	{
+		ResolveInterfacePtr();
+	}
+
+	if (WeakObject.IsValid())
+	{
+		InterfacePtr->OnStopPlaying(InStopTime);
+	}
+}
+
+FFrameTime FMovieSceneTimeController_Custom::OnRequestCurrentTime(const FQualifiedFrameTime& InCurrentTime, float InPlayRate)
+{
+	if (WeakObject.IsStale())
+	{
+		ResolveInterfacePtr();
+	}
+
+	if (WeakObject.IsValid())
+	{
+		return InterfacePtr->OnRequestCurrentTime(InCurrentTime, InPlayRate);
+	}
+
+	return InCurrentTime.Time;
+}
+
+void FMovieSceneTimeController_Custom::ResolveInterfacePtr()
+{
+	WeakObject = nullptr;
+	InterfacePtr = nullptr;
+
+	if (UObject* ResolvedClockSource = ObjectPath.ResolveObject())
+	{
+		const bool bHasCompatibleInterface = ResolvedClockSource->GetClass()->ImplementsInterface(UMovieSceneCustomClockSource::StaticClass());
+		if (bHasCompatibleInterface)
+		{
+			WeakObject = ResolvedClockSource;
+			InterfacePtr = static_cast<IMovieSceneCustomClockSource*>(ResolvedClockSource->GetInterfaceAddress(UMovieSceneCustomClockSource::StaticClass()));
+		}
+	}
+}
\ No newline at end of file
diff --git a/Engine/Source/Runtime/MovieScene/Private/MovieScene.cpp b/Engine/Source/Runtime/MovieScene/Private/MovieScene.cpp
--- a/Engine/Source/Runtime/MovieScene/Private/MovieScene.cpp
+++ b/Engine/Source/Runtime/MovieScene/Private/MovieScene.cpp
@@ -7,6 +7,7 @@
 #include "Evaluation/MovieSceneEvaluationCustomVersion.h"
 #include "Compilation/MovieSceneSegmentCompiler.h"
 #include "UObject/SequencerObjectVersion.h"
+#include "Evaluation/IMovieSceneCustomClockSource.h"
 #include "CommonFrameRates.h"
 
 #define LOCTEXT_NAMESPACE "MovieScene"
@@ -1243,6 +1244,11 @@ void UMovieScene::MoveBindingContents(const FGuid& SourceBindingId, const FGuid&
 	}
 }
 
+TSharedPtr<FMovieSceneTimeController> UMovieScene::MakeCustomTimeController(UObject* PlaybackContext)
+{
+	return MakeShared<FMovieSceneTimeController_Custom>(CustomClockSourcePath, PlaybackContext);
+}
+
 void UMovieScene::SetMarkedFrame(int32 InMarkIndex, FFrameNumber InFrameNumber)
 {
 	if (InMarkIndex < 0 && InMarkIndex > MarkedFrames.Num()-1)
diff --git a/Engine/Source/Runtime/MovieScene/Private/MovieSceneSequencePlayer.cpp b/Engine/Source/Runtime/MovieScene/Private/MovieSceneSequencePlayer.cpp
--- a/Engine/Source/Runtime/MovieScene/Private/MovieSceneSequencePlayer.cpp
+++ b/Engine/Source/Runtime/MovieScene/Private/MovieSceneSequencePlayer.cpp
@@ -665,6 +665,11 @@ void UMovieSceneSequencePlayer::Initialize(UMovieSceneSequence* InSequence, cons
 		StartTimeWithOffset = StartTime + StartingTimeOffset;
 
 		ClockToUse = MovieScene->GetClockSource();
+
+		if (ClockToUse == EUpdateClockSource::Custom)
+		{
+			TimeController = MovieScene->MakeCustomTimeController(GetPlaybackContext());
+		}
 	}
 
 	if (!TimeController.IsValid())
@@ -676,6 +681,11 @@ void UMovieSceneSequencePlayer::Initialize(UMovieSceneSequence* InSequence, cons
 		case EUpdateClockSource::Timecode: TimeController = MakeShared<FMovieSceneTimeController_TimecodeClock>(); break;
 		default:                           TimeController = MakeShared<FMovieSceneTimeController_Tick>();          break;
 		}
+
+		if (!ensureMsgf(TimeController.IsValid(), TEXT("No time controller specified for sequence playback. Falling back to Engine Tick clock source.")))
+		{
+			TimeController = MakeShared<FMovieSceneTimeController_Tick>();
+		}
 	}
 
 	RootTemplateInstance.Initialize(*Sequence, *this);
diff --git a/Engine/Source/Runtime/MovieScene/Public/Evaluation/IMovieSceneCustomClockSource.h b/Engine/Source/Runtime/MovieScene/Public/Evaluation/IMovieSceneCustomClockSource.h
new file mode 100644
--- /dev/null
+++ b/Engine/Source/Runtime/MovieScene/Public/Evaluation/IMovieSceneCustomClockSource.h
@@ -0,0 +1,62 @@
+// Copyright 1998-2020 Epic Games, Inc. All Rights Reserved.
+
+#pragma once
+
+#include "UObject/Interface.h"
+#include "MovieSceneTimeController.h"
+#include "UObject/WeakObjectPtrTemplates.h"
+#include "IMovieSceneCustomClockSource.generated.h"
+
+struct FQualifiedFrameTime;
+
+UINTERFACE()
+class MOVIESCENE_API UMovieSceneCustomClockSource : public UInterface
+{
+	GENERATED_BODY()
+};
+
+/**
+ * 
+ */
+class IMovieSceneCustomClockSource
+{
+public:
+
+	GENERATED_BODY()
+
+	UFUNCTION()
+	virtual void OnTick(float DeltaSeconds, float InPlayRate) {}
+
+	UFUNCTION()
+	virtual void OnStartPlaying(const FQualifiedFrameTime& InStartTime) {}
+
+	UFUNCTION()
+	virtual void OnStopPlaying(const FQualifiedFrameTime& InStopTime) {}
+
+	UFUNCTION()
+	virtual FFrameTime OnRequestCurrentTime(const FQualifiedFrameTime& InCurrentTime, float InPlayRate) { return 0; }
+};
+
+struct MOVIESCENE_API FMovieSceneTimeController_Custom : FMovieSceneTimeController
+{
+	explicit FMovieSceneTimeController_Custom(const FSoftObjectPath& InObjectPath, TWeakObjectPtr<> PlaybackContext);
+
+private:
+
+	virtual void OnTick(float DeltaSeconds, float InPlayRate) override final;
+	virtual void OnStartPlaying(const FQualifiedFrameTime& InStartTime) override final;
+	virtual void OnStopPlaying(const FQualifiedFrameTime& InStopTime) override final;
+	virtual FFrameTime OnRequestCurrentTime(const FQualifiedFrameTime& InCurrentTime, float InPlayRate) override final;
+
+private:
+
+	void ResolveInterfacePtr();
+
+	TWeakObjectPtr<> WeakPlaybackContext;
+
+	TWeakObjectPtr<> WeakObject;
+
+	IMovieSceneCustomClockSource* InterfacePtr;
+
+	FSoftObjectPath ObjectPath;
+};
\ No newline at end of file
diff --git a/Engine/Source/Runtime/MovieScene/Public/MovieScene.h b/Engine/Source/Runtime/MovieScene/Public/MovieScene.h
--- a/Engine/Source/Runtime/MovieScene/Public/MovieScene.h
+++ b/Engine/Source/Runtime/MovieScene/Public/MovieScene.h
@@ -19,6 +19,8 @@
 #include "MovieSceneTimeController.h"
 #include "MovieScene.generated.h"
 
+struct FMovieSceneTimeController;
+
 class UMovieSceneFolder;
 class UMovieSceneSection;
 class UMovieSceneTrack;
@@ -687,12 +689,30 @@ public:
 		return ClockSource;
 	}
 
+	/**
+	 * Retrieve a time controller from this sequence instance, if the clock source is set to custom
+	 */
+	TSharedPtr<FMovieSceneTimeController> MakeCustomTimeController(UObject* PlaybackContext);
+
 	/**
 	 * Assign the clock source to be used for this moviescene
 	 */
 	void SetClockSource(EUpdateClockSource InNewClockSource)
 	{
 		ClockSource = InNewClockSource;
+		if (ClockSource != EUpdateClockSource::Custom)
+		{
+			CustomClockSourcePath.Reset();
+		}
+	}
+
+	/**
+	 * Assign the clock source to be used for this moviescene
+	 */
+	void SetClockSource(UObject* InNewClockSource)
+	{
+		ClockSource = EUpdateClockSource::Custom;
+		CustomClockSourcePath = InNewClockSource;
 	}
 
 	/*
@@ -1015,6 +1035,9 @@ private:
 	UPROPERTY()
 	EUpdateClockSource ClockSource;
 
+	UPROPERTY()
+	FSoftObjectPath CustomClockSourcePath;
+
 	/** The set of user-marked frames */
 	UPROPERTY()
 	TArray<FMovieSceneMarkedFrame> MarkedFrames;
diff --git a/Engine/Source/Runtime/MovieScene/Public/MovieSceneFwd.h b/Engine/Source/Runtime/MovieScene/Public/MovieSceneFwd.h
--- a/Engine/Source/Runtime/MovieScene/Public/MovieSceneFwd.h
+++ b/Engine/Source/Runtime/MovieScene/Public/MovieSceneFwd.h
@@ -50,6 +50,9 @@ enum class EUpdateClockSource : uint8
 
 	/** Use current timecode provider for timing. Does not honor world or actor pause state. */
 	Timecode,
+
+	/** Custom clock source created and defined externally. */
+	Custom,
 };
 
 MOVIESCENE_API DECLARE_LOG_CATEGORY_EXTERN(LogMovieScene, Log, All);
