Fix deprecation warning
#rnx

diff --git a/Engine/Plugins/Media/MediaCompositing/Source/MediaCompositingEditor/Private/Sequencer/MediaPlayerPropertyTrackEditor.cpp b/Engine/Plugins/Media/MediaCompositing/Source/MediaCompositingEditor/Private/Sequencer/MediaPlayerPropertyTrackEditor.cpp
--- a/Engine/Plugins/Media/MediaCompositing/Source/MediaCompositingEditor/Private/Sequencer/MediaPlayerPropertyTrackEditor.cpp
+++ b/Engine/Plugins/Media/MediaCompositing/Source/MediaCompositingEditor/Private/Sequencer/MediaPlayerPropertyTrackEditor.cpp
@@ -101,7 +101,7 @@ void FMediaPlayerPropertyTrackEditor::OnAnimatedPropertyChanged(const FPropertyC
 		return;
 	}
 
-	UProperty* Property = PropertyChangedParams.PropertyPath.GetLeafMostProperty().Property.Get();
+	FProperty* Property = PropertyChangedParams.PropertyPath.GetLeafMostProperty().Property.Get();
 	if (!Property)
 	{
 		return;
