Updating copyright notices to the new standard

diff --git a/Engine/Plugins/Media/MediaCompositing/Source/MediaCompositing/Private/MovieScene/MovieSceneMediaPlayerPropertySection.cpp b/Engine/Plugins/Media/MediaCompositing/Source/MediaCompositing/Private/MovieScene/MovieSceneMediaPlayerPropertySection.cpp
--- a/Engine/Plugins/Media/MediaCompositing/Source/MediaCompositing/Private/MovieScene/MovieSceneMediaPlayerPropertySection.cpp
+++ b/Engine/Plugins/Media/MediaCompositing/Source/MediaCompositing/Private/MovieScene/MovieSceneMediaPlayerPropertySection.cpp
@@ -1,4 +1,4 @@
-// Copyright 1998-2020 Epic Games, Inc. All Rights Reserved.
+// Copyright Epic Games, Inc. All Rights Reserved.
 
 #include "MovieSceneMediaPlayerPropertySection.h"
 #include "MovieScene.h"
diff --git a/Engine/Plugins/Media/MediaCompositing/Source/MediaCompositing/Private/MovieScene/MovieSceneMediaPlayerPropertyTemplate.cpp b/Engine/Plugins/Media/MediaCompositing/Source/MediaCompositing/Private/MovieScene/MovieSceneMediaPlayerPropertyTemplate.cpp
--- a/Engine/Plugins/Media/MediaCompositing/Source/MediaCompositing/Private/MovieScene/MovieSceneMediaPlayerPropertyTemplate.cpp
+++ b/Engine/Plugins/Media/MediaCompositing/Source/MediaCompositing/Private/MovieScene/MovieSceneMediaPlayerPropertyTemplate.cpp
@@ -1,4 +1,4 @@
-// Copyright 1998-2020 Epic Games, Inc. All Rights Reserved.
+// Copyright Epic Games, Inc. All Rights Reserved.
 
 #include "MovieSceneMediaPlayerPropertyTemplate.h"
 
diff --git a/Engine/Plugins/Media/MediaCompositing/Source/MediaCompositing/Private/MovieScene/MovieSceneMediaPlayerPropertyTemplate.h b/Engine/Plugins/Media/MediaCompositing/Source/MediaCompositing/Private/MovieScene/MovieSceneMediaPlayerPropertyTemplate.h
--- a/Engine/Plugins/Media/MediaCompositing/Source/MediaCompositing/Private/MovieScene/MovieSceneMediaPlayerPropertyTemplate.h
+++ b/Engine/Plugins/Media/MediaCompositing/Source/MediaCompositing/Private/MovieScene/MovieSceneMediaPlayerPropertyTemplate.h
@@ -1,4 +1,4 @@
-// Copyright 1998-2020 Epic Games, Inc. All Rights Reserved.
+// Copyright Epic Games, Inc. All Rights Reserved.
 
 #pragma once
 
diff --git a/Engine/Plugins/Media/MediaCompositing/Source/MediaCompositing/Private/MovieScene/MovieSceneMediaPlayerPropertyTrack.cpp b/Engine/Plugins/Media/MediaCompositing/Source/MediaCompositing/Private/MovieScene/MovieSceneMediaPlayerPropertyTrack.cpp
--- a/Engine/Plugins/Media/MediaCompositing/Source/MediaCompositing/Private/MovieScene/MovieSceneMediaPlayerPropertyTrack.cpp
+++ b/Engine/Plugins/Media/MediaCompositing/Source/MediaCompositing/Private/MovieScene/MovieSceneMediaPlayerPropertyTrack.cpp
@@ -1,4 +1,4 @@
-// Copyright 1998-2020 Epic Games, Inc. All Rights Reserved.
+// Copyright Epic Games, Inc. All Rights Reserved.
 
 #include "MovieSceneMediaPlayerPropertyTrack.h"
 
diff --git a/Engine/Plugins/Media/MediaCompositing/Source/MediaCompositing/Public/MovieSceneMediaPlayerPropertySection.h b/Engine/Plugins/Media/MediaCompositing/Source/MediaCompositing/Public/MovieSceneMediaPlayerPropertySection.h
--- a/Engine/Plugins/Media/MediaCompositing/Source/MediaCompositing/Public/MovieSceneMediaPlayerPropertySection.h
+++ b/Engine/Plugins/Media/MediaCompositing/Source/MediaCompositing/Public/MovieSceneMediaPlayerPropertySection.h
@@ -1,4 +1,4 @@
-// Copyright 1998-2020 Epic Games, Inc. All Rights Reserved.
+// Copyright Epic Games, Inc. All Rights Reserved.
 
 #pragma once
 
diff --git a/Engine/Plugins/Media/MediaCompositing/Source/MediaCompositing/Public/MovieSceneMediaPlayerPropertyTrack.h b/Engine/Plugins/Media/MediaCompositing/Source/MediaCompositing/Public/MovieSceneMediaPlayerPropertyTrack.h
--- a/Engine/Plugins/Media/MediaCompositing/Source/MediaCompositing/Public/MovieSceneMediaPlayerPropertyTrack.h
+++ b/Engine/Plugins/Media/MediaCompositing/Source/MediaCompositing/Public/MovieSceneMediaPlayerPropertyTrack.h
@@ -1,4 +1,4 @@
-// Copyright 1998-2020 Epic Games, Inc. All Rights Reserved.
+// Copyright Epic Games, Inc. All Rights Reserved.
 
 #pragma once
 
diff --git a/Engine/Plugins/Media/MediaCompositing/Source/MediaCompositingEditor/Private/Sequencer/MediaPlayerPropertyTrackEditor.cpp b/Engine/Plugins/Media/MediaCompositing/Source/MediaCompositingEditor/Private/Sequencer/MediaPlayerPropertyTrackEditor.cpp
--- a/Engine/Plugins/Media/MediaCompositing/Source/MediaCompositingEditor/Private/Sequencer/MediaPlayerPropertyTrackEditor.cpp
+++ b/Engine/Plugins/Media/MediaCompositing/Source/MediaCompositingEditor/Private/Sequencer/MediaPlayerPropertyTrackEditor.cpp
@@ -1,4 +1,4 @@
-// Copyright 1998-2020 Epic Games, Inc. All Rights Reserved.
+// Copyright Epic Games, Inc. All Rights Reserved.
 
 #include "MediaPlayerPropertyTrackEditor.h"
 
diff --git a/Engine/Plugins/Media/MediaCompositing/Source/MediaCompositingEditor/Private/Sequencer/MediaPlayerPropertyTrackEditor.h b/Engine/Plugins/Media/MediaCompositing/Source/MediaCompositingEditor/Private/Sequencer/MediaPlayerPropertyTrackEditor.h
--- a/Engine/Plugins/Media/MediaCompositing/Source/MediaCompositingEditor/Private/Sequencer/MediaPlayerPropertyTrackEditor.h
+++ b/Engine/Plugins/Media/MediaCompositing/Source/MediaCompositingEditor/Private/Sequencer/MediaPlayerPropertyTrackEditor.h
@@ -1,4 +1,4 @@
-// Copyright 1998-2020 Epic Games, Inc. All Rights Reserved.
+// Copyright Epic Games, Inc. All Rights Reserved.
 
 #pragma once
 
diff --git a/Engine/Source/Runtime/LevelSequence/Private/SequenceMediaController.cpp b/Engine/Source/Runtime/LevelSequence/Private/SequenceMediaController.cpp
--- a/Engine/Source/Runtime/LevelSequence/Private/SequenceMediaController.cpp
+++ b/Engine/Source/Runtime/LevelSequence/Private/SequenceMediaController.cpp
@@ -1,4 +1,4 @@
-// Copyright 1998-2020 Epic Games, Inc. All Rights Reserved.
+// Copyright Epic Games, Inc. All Rights Reserved.
 
 #include "SequenceMediaController.h"
 
diff --git a/Engine/Source/Runtime/LevelSequence/Public/SequenceMediaController.h b/Engine/Source/Runtime/LevelSequence/Public/SequenceMediaController.h
--- a/Engine/Source/Runtime/LevelSequence/Public/SequenceMediaController.h
+++ b/Engine/Source/Runtime/LevelSequence/Public/SequenceMediaController.h
@@ -1,4 +1,4 @@
-// Copyright 1998-2020 Epic Games, Inc. All Rights Reserved.
+// Copyright Epic Games, Inc. All Rights Reserved.
 
 #pragma once
 
diff --git a/Engine/Source/Runtime/MediaAssets/Private/Assets/MediaComponent.cpp b/Engine/Source/Runtime/MediaAssets/Private/Assets/MediaComponent.cpp
--- a/Engine/Source/Runtime/MediaAssets/Private/Assets/MediaComponent.cpp
+++ b/Engine/Source/Runtime/MediaAssets/Private/Assets/MediaComponent.cpp
@@ -1,4 +1,4 @@
-// Copyright 1998-2020 Epic Games, Inc. All Rights Reserved.
+// Copyright Epic Games, Inc. All Rights Reserved.
 
 #include "MediaComponent.h"
 #include "MediaPlayer.h"
diff --git a/Engine/Source/Runtime/MediaAssets/Public/MediaComponent.h b/Engine/Source/Runtime/MediaAssets/Public/MediaComponent.h
--- a/Engine/Source/Runtime/MediaAssets/Public/MediaComponent.h
+++ b/Engine/Source/Runtime/MediaAssets/Public/MediaComponent.h
@@ -1,4 +1,4 @@
-// Copyright 1998-2020 Epic Games, Inc. All Rights Reserved.
+// Copyright Epic Games, Inc. All Rights Reserved.
 
 #pragma once
 
diff --git a/Engine/Source/Runtime/MovieScene/Private/Evaluation/IMovieSceneCustomClockSource.cpp b/Engine/Source/Runtime/MovieScene/Private/Evaluation/IMovieSceneCustomClockSource.cpp
--- a/Engine/Source/Runtime/MovieScene/Private/Evaluation/IMovieSceneCustomClockSource.cpp
+++ b/Engine/Source/Runtime/MovieScene/Private/Evaluation/IMovieSceneCustomClockSource.cpp
@@ -1,4 +1,4 @@
-// Copyright 1998-2020 Epic Games, Inc. All Rights Reserved.
+// Copyright Epic Games, Inc. All Rights Reserved.
 
 #include "Evaluation/IMovieSceneCustomClockSource.h"
 #include "UObject/Package.h"
diff --git a/Engine/Source/Runtime/MovieScene/Public/Evaluation/IMovieSceneCustomClockSource.h b/Engine/Source/Runtime/MovieScene/Public/Evaluation/IMovieSceneCustomClockSource.h
--- a/Engine/Source/Runtime/MovieScene/Public/Evaluation/IMovieSceneCustomClockSource.h
+++ b/Engine/Source/Runtime/MovieScene/Public/Evaluation/IMovieSceneCustomClockSource.h
@@ -1,4 +1,4 @@
-// Copyright 1998-2020 Epic Games, Inc. All Rights Reserved.
+// Copyright Epic Games, Inc. All Rights Reserved.
 
 #pragma once
 
