Fixing potentially uninitialized members

diff --git a/Engine/Plugins/Media/MediaCompositing/Source/MediaCompositing/Private/MovieScene/MovieSceneMediaPlayerPropertyTemplate.h b/Engine/Plugins/Media/MediaCompositing/Source/MediaCompositing/Private/MovieScene/MovieSceneMediaPlayerPropertyTemplate.h
--- a/Engine/Plugins/Media/MediaCompositing/Source/MediaCompositing/Private/MovieScene/MovieSceneMediaPlayerPropertyTemplate.h
+++ b/Engine/Plugins/Media/MediaCompositing/Source/MediaCompositing/Private/MovieScene/MovieSceneMediaPlayerPropertyTemplate.h
@@ -18,7 +18,7 @@ struct FMovieSceneMediaPlayerPropertySectionTemplate
 	GENERATED_BODY()
 
 	/** Default constructor. */
-	FMovieSceneMediaPlayerPropertySectionTemplate() { }
+	FMovieSceneMediaPlayerPropertySectionTemplate() : MediaSource(nullptr), bLoop(false) { }
 	FMovieSceneMediaPlayerPropertySectionTemplate(const UMovieSceneMediaPlayerPropertySection* InSection, const UMovieSceneMediaPlayerPropertyTrack* InTrack);
 
 public:
