Adjusting movie playback logic in sequencer to fully work with new "V2" timing support of ImgPlayer & MFW

diff --git a/Engine/Plugins/Media/MediaCompositing/Source/MediaCompositing/Private/MovieScene/MovieSceneMediaData.cpp b/Engine/Plugins/Media/MediaCompositing/Source/MediaCompositing/Private/MovieScene/MovieSceneMediaData.cpp
--- a/Engine/Plugins/Media/MediaCompositing/Source/MediaCompositing/Private/MovieScene/MovieSceneMediaData.cpp
+++ b/Engine/Plugins/Media/MediaCompositing/Source/MediaCompositing/Private/MovieScene/MovieSceneMediaData.cpp
@@ -80,7 +80,17 @@ void FMovieSceneMediaData::HandleMediaPlayerEvent(EMediaEvent Event)
 		return; // media can't seek
 	}
 
-	const FTimespan MediaTime = SeekOnOpenTime % MediaPlayer->GetDuration();
+	const FTimespan Duration = MediaPlayer->GetDuration();
+	FTimespan MediaTime;
+
+	if (!MediaPlayer->IsLooping())
+	{
+		MediaTime = FMath::Clamp(SeekOnOpenTime, FTimespan::Zero(), Duration - FTimespan::FromSeconds(0.001));
+	}
+	else
+	{
+		MediaTime = SeekOnOpenTime % Duration;
+	}
 
 	MediaPlayer->SetRate(0.0f);
 	MediaPlayer->Seek(MediaTime);
diff --git a/Engine/Plugins/Media/MediaCompositing/Source/MediaCompositing/Private/MovieScene/MovieSceneMediaTemplate.cpp b/Engine/Plugins/Media/MediaCompositing/Source/MediaCompositing/Private/MovieScene/MovieSceneMediaTemplate.cpp
--- a/Engine/Plugins/Media/MediaCompositing/Source/MediaCompositing/Private/MovieScene/MovieSceneMediaTemplate.cpp
+++ b/Engine/Plugins/Media/MediaCompositing/Source/MediaCompositing/Private/MovieScene/MovieSceneMediaTemplate.cpp
@@ -60,8 +60,9 @@ private:
 struct FMediaSectionExecutionToken
 	: IMovieSceneExecutionToken
 {
-	FMediaSectionExecutionToken(UMediaSource* InMediaSource, FTimespan InCurrentTime)
+	FMediaSectionExecutionToken(UMediaSource* InMediaSource, FTimespan InCurrentTime, FTimespan InFrameDuration)
 		: CurrentTime(InCurrentTime)
+		, FrameDuration(InFrameDuration)
 		, MediaSource(InMediaSource)
 		, PlaybackRate(1.0f)
 	{ }
@@ -100,8 +101,22 @@ struct FMediaSectionExecutionToken
 			return; // media has no length
 		}
 
+		//
 		// update media player
-		const FTimespan MediaTime = CurrentTime % MediaDuration;
+		//
+
+		// Setup media time (used for seeks)
+		FTimespan MediaTime;
+		if (!MediaPlayer->IsLooping())
+		{
+			// note: we use a small offset at the end to make sure we can indeed seek to it (exclusive end type range)
+			MediaTime = FMath::Clamp(CurrentTime, FTimespan::Zero(), MediaDuration - FrameDuration * 0.5);
+		}
+		else
+		{
+			// one always seeks into the original media time-range, hence: modulo the time
+			MediaTime = CurrentTime % MediaDuration;
+		}
 
 		#if MOVIESCENEMEDIATEMPLATE_TRACE_EVALUATION
 			GLog->Logf(ELogVerbosity::Log, TEXT("Executing time %s, MediaTime %s"), *CurrentTime.ToString(TEXT("%h:%m:%s.%t")), *MediaTime.ToString(TEXT("%h:%m:%s.%t")));
@@ -112,14 +127,47 @@ struct FMediaSectionExecutionToken
 			if (!MediaPlayer->IsPlaying())
 			{
 				MediaPlayer->Seek(MediaTime);
-				MediaPlayer->SetRate(1.0f);
+				// Set rate
+				// (note that the DIRECTION is important, but the magnitude is not - as we use blocked playback, the range setup to block on will serve as external clock to the player,
+				//  the direction is taken into account as hint for internal operation of the player)
+				if (!MediaPlayer->SetRate((Context.GetDirection() == EPlayDirection::Forwards) ? 1.0f : -1.0f))
+				{
+					// Failed to set needed rate: better switch off blocking and bail...
+					MediaPlayer->SetBlockOnTimeRange(TRange<FTimespan>::Empty());
+					return;
+				}
 			}
-			else if (Context.HasJumped())
+			else
 			{
-				MediaPlayer->Seek(MediaTime);
+				if (Context.HasJumped())
+				{
+					MediaPlayer->Seek(MediaTime);
+				}
+
+				float CurrentPlayerRate = MediaPlayer->GetRate();
+				if (Context.GetDirection() == EPlayDirection::Forwards && CurrentPlayerRate < 0.0f)
+				{
+					if (!MediaPlayer->SetRate(1.0f))
+					{
+						// Failed to set needed rate: better switch off blocking and bail...
+						MediaPlayer->SetBlockOnTimeRange(TRange<FTimespan>::Empty());
+						return;
+					}
+				}
+				else if (Context.GetDirection() == EPlayDirection::Backwards && CurrentPlayerRate > 0.0f)
+				{
+					if (!MediaPlayer->SetRate(-1.0f))
+					{
+						// Failed to set needed rate: better switch off blocking and bail...
+						MediaPlayer->SetBlockOnTimeRange(TRange<FTimespan>::Empty());
+						return;
+					}
+				}
 			}
 
-			MediaPlayer->SetBlockOnTime(MediaPlayer->GetTime());
+			// Set blocking range / time-range to display
+			// (we always use the full current time for this, any adjustments to player timestamps are done internally)
+			MediaPlayer->SetBlockOnTimeRange(TRange<FTimespan>(CurrentTime, CurrentTime + FrameDuration));
 		}
 		else
 		{
@@ -129,13 +177,14 @@ struct FMediaSectionExecutionToken
 			}
 
 			MediaPlayer->Seek(MediaTime);
-			MediaPlayer->SetBlockOnTime(FTimespan::MinValue());
+			MediaPlayer->SetBlockOnTimeRange(TRange<FTimespan>::Empty());
 		}
 	}
 
 private:
 
 	FTimespan CurrentTime;
+	FTimespan FrameDuration;
 	UMediaSource* MediaSource;
 	float PlaybackRate;
 };
@@ -194,6 +243,7 @@ void FMovieSceneMediaSectionTemplate::Evaluate(const FMovieSceneEvaluationOperan
 		const int64 DenominatorTicks = FrameRate.Denominator * ETimespan::TicksPerSecond;
 		const int64 FrameTicks = FMath::DivideAndRoundNearest(int64(FrameTime.GetFrame().Value * DenominatorTicks), int64(FrameRate.Numerator));
 		const int64 FrameSubTicks = FMath::DivideAndRoundNearest(int64(FrameTime.GetSubFrame() * DenominatorTicks), int64(FrameRate.Numerator));
+		const int64 FrameDurationTicks = 1000 * FMath::DivideAndRoundNearest(DenominatorTicks, int64(FrameRate.Numerator));
 
 		#if MOVIESCENEMEDIATEMPLATE_TRACE_EVALUATION
 			GLog->Logf(ELogVerbosity::Log, TEXT("Evaluating frame %i+%f, FrameRate %i/%i, FrameTicks %d+%d"),
@@ -206,7 +256,7 @@ void FMovieSceneMediaSectionTemplate::Evaluate(const FMovieSceneEvaluationOperan
 			);
 		#endif
 
-		ExecutionTokens.Add(FMediaSectionExecutionToken(Params.MediaSource, FTimespan(FrameTicks + FrameSubTicks)));
+		ExecutionTokens.Add(FMediaSectionExecutionToken(Params.MediaSource, FTimespan(FrameTicks + FrameSubTicks), FTimespan(FrameDurationTicks)));
 	}
 }
 
