Sequencer: bLooping true by default for media sections since bug in media track was masking this behavior.

diff --git a/Engine/Plugins/Media/MediaCompositing/Source/MediaCompositing/Private/MovieScene/MovieSceneMediaSection.cpp b/Engine/Plugins/Media/MediaCompositing/Source/MediaCompositing/Private/MovieScene/MovieSceneMediaSection.cpp
--- a/Engine/Plugins/Media/MediaCompositing/Source/MediaCompositing/Private/MovieScene/MovieSceneMediaSection.cpp
+++ b/Engine/Plugins/Media/MediaCompositing/Source/MediaCompositing/Private/MovieScene/MovieSceneMediaSection.cpp
@@ -11,6 +11,7 @@
 
 UMovieSceneMediaSection::UMovieSceneMediaSection(const FObjectInitializer& ObjectInitializer)
 	: Super(ObjectInitializer)
+	, bLooping(true)
 {
 #if WITH_EDITORONLY_DATA
 	ThumbnailReferenceOffset = 0.f;
