Changed movie player control code in Sequencer to make sure blocking is used more often & fixed up some details
#rnx

diff --git a/Engine/Plugins/Media/MediaCompositing/Source/MediaCompositing/Private/MovieScene/MovieSceneMediaData.cpp b/Engine/Plugins/Media/MediaCompositing/Source/MediaCompositing/Private/MovieScene/MovieSceneMediaData.cpp
--- a/Engine/Plugins/Media/MediaCompositing/Source/MediaCompositing/Private/MovieScene/MovieSceneMediaData.cpp
+++ b/Engine/Plugins/Media/MediaCompositing/Source/MediaCompositing/Private/MovieScene/MovieSceneMediaData.cpp
@@ -81,8 +81,12 @@ void FMovieSceneMediaData::HandleMediaPlayerEvent(EMediaEvent Event)
 	}
 
 	const FTimespan Duration = MediaPlayer->GetDuration();
-	FTimespan MediaTime;
+    if (Duration == FTimespan::Zero())
+    {
+	    return;
+    }
 
+	FTimespan MediaTime;
 	if (!MediaPlayer->IsLooping())
 	{
 		MediaTime = FMath::Clamp(SeekOnOpenTime, FTimespan::Zero(), Duration - FTimespan::FromSeconds(0.001));
diff --git a/Engine/Plugins/Media/MediaCompositing/Source/MediaCompositing/Private/MovieScene/MovieSceneMediaTemplate.cpp b/Engine/Plugins/Media/MediaCompositing/Source/MediaCompositing/Private/MovieScene/MovieSceneMediaTemplate.cpp
--- a/Engine/Plugins/Media/MediaCompositing/Source/MediaCompositing/Private/MovieScene/MovieSceneMediaTemplate.cpp
+++ b/Engine/Plugins/Media/MediaCompositing/Source/MediaCompositing/Private/MovieScene/MovieSceneMediaTemplate.cpp
@@ -81,12 +81,15 @@ struct FMediaSectionExecutionToken
 		if (MediaPlayer->GetUrl().IsEmpty())
 		{
 			SectionData.SeekOnOpen(CurrentTime);
+			// Setup an initial blocking range - MediaFramework will block (even through the opening process) in its next tick...
+			MediaPlayer->SetBlockOnTimeRange(TRange<FTimespan>(CurrentTime, CurrentTime + FrameDuration));
 			MediaPlayer->OpenSource(MediaSource);
 
 			return;
 		}
 
 		// seek on open if necessary
+		// (usually should not be needed as the blocking on open should ensure we never see the player preparing here)
 		if (MediaPlayer->IsPreparing())
 		{
 			SectionData.SeekOnOpen(CurrentTime);
@@ -164,10 +167,6 @@ struct FMediaSectionExecutionToken
 					}
 				}
 			}
-
-			// Set blocking range / time-range to display
-			// (we always use the full current time for this, any adjustments to player timestamps are done internally)
-			MediaPlayer->SetBlockOnTimeRange(TRange<FTimespan>(CurrentTime, CurrentTime + FrameDuration));
 		}
 		else
 		{
@@ -177,8 +176,11 @@ struct FMediaSectionExecutionToken
 			}
 
 			MediaPlayer->Seek(MediaTime);
-			MediaPlayer->SetBlockOnTimeRange(TRange<FTimespan>::Empty());
 		}
+
+	// Set blocking range / time-range to display
+	// (we always use the full current time for this, any adjustments to player timestamps are done internally)
+		MediaPlayer->SetBlockOnTimeRange(TRange<FTimespan>(CurrentTime, CurrentTime + FrameDuration));
 	}
 
 private:
