Adds a Sequencer Timeline and synchronizes changes in the asset editor with Sequencer and vice-versa

diff --git a/Engine/Plugins/Media/MediaCompositing/Source/MediaCompositingEditor/Private/MediaCompositingEditor.cpp b/Engine/Plugins/Media/MediaCompositing/Source/MediaCompositingEditor/Private/MediaCompositingEditor.cpp
--- a/Engine/Plugins/Media/MediaCompositing/Source/MediaCompositingEditor/Private/MediaCompositingEditor.cpp
+++ b/Engine/Plugins/Media/MediaCompositing/Source/MediaCompositingEditor/Private/MediaCompositingEditor.cpp
@@ -7,7 +7,7 @@
 #include "ISequenceRecorder.h"
 #include "MediaCompositingEditorStyle.h"
 #include "MediaSequenceRecorderExtender.h"
-#include "MediaTrackEditor.h"
+#include "Sequencer/MediaTrackEditor.h"
 #include "MediaPlayerPropertyTrackEditor.h"
 
 
diff --git a/Engine/Plugins/Media/MediaCompositing/Source/MediaCompositingEditor/Private/Sequencer/MediaTrackEditor.cpp b/Engine/Plugins/Media/MediaCompositing/Source/MediaCompositingEditor/Private/Sequencer/MediaTrackEditor.cpp
--- a/Engine/Plugins/Media/MediaCompositing/Source/MediaCompositingEditor/Private/Sequencer/MediaTrackEditor.cpp
+++ b/Engine/Plugins/Media/MediaCompositing/Source/MediaCompositingEditor/Private/Sequencer/MediaTrackEditor.cpp
@@ -1,6 +1,6 @@
 // Copyright Epic Games, Inc. All Rights Reserved.
 
-#include "MediaTrackEditor.h"
+#include "Sequencer/MediaTrackEditor.h"
 
 #include "ContentBrowserModule.h"
 #include "EditorStyleSet.h"
diff --git a/Engine/Plugins/Media/MediaCompositing/Source/MediaCompositingEditor/Private/Sequencer/MediaTrackEditor.h b/Engine/Plugins/Media/MediaCompositing/Source/MediaCompositingEditor/Private/Sequencer/MediaTrackEditor.h
deleted file mode 100644
diff --git a/Engine/Plugins/Media/MediaCompositing/Source/MediaCompositingEditor/Public/Sequencer/MediaTrackEditor.h b/Engine/Plugins/Media/MediaCompositing/Source/MediaCompositingEditor/Public/Sequencer/MediaTrackEditor.h
new file mode 100644
--- /dev/null
+++ b/Engine/Plugins/Media/MediaCompositing/Source/MediaCompositingEditor/Public/Sequencer/MediaTrackEditor.h
@@ -0,0 +1,88 @@
+// Copyright Epic Games, Inc. All Rights Reserved.
+
+#pragma once
+
+#include "AnimatedPropertyKey.h"
+#include "AssetData.h"
+#include "MovieSceneTrackEditor.h"
+#include "Templates/SharedPointer.h"
+
+class ISequencer;
+class FTrackEditorThumbnailPool;
+class UMediaSource;
+class UMovieSceneMediaTrack;
+
+
+/**
+ * Track editor that understands how to animate MediaPlayer properties on objects
+ */
+class MEDIACOMPOSITINGEDITOR_API FMediaTrackEditor
+	: public FMovieSceneTrackEditor
+{
+public:
+
+	/**
+	 * Create a new media track editor instance.
+	 *
+	 * @param OwningSequencer The sequencer object that will own the track editor.
+	 * @return The new track editor.
+	 */
+	static TSharedRef<ISequencerTrackEditor>  CreateTrackEditor(TSharedRef<ISequencer> OwningSequencer)
+	{
+		return MakeShared<FMediaTrackEditor>(OwningSequencer);
+	}
+
+	/**
+	 * Get the list of all property types that this track editor animates.
+	 *
+	 * @return List of animated properties.
+	 */
+	static TArray<FAnimatedPropertyKey, TInlineAllocator<1>> GetAnimatedPropertyTypes();
+
+public:
+
+	/**
+	 * Create and initialize a new instance.
+	 *
+	 * @param InSequencer The sequencer object that owns this track editor.
+	 */
+	FMediaTrackEditor(TSharedRef<ISequencer> InSequencer);
+
+	/** Virtual destructor. */
+	virtual ~FMediaTrackEditor();
+
+public:
+
+	//~ FMovieSceneTrackEditor interface
+
+	virtual UMovieSceneTrack* AddTrack(UMovieScene* FocusedMovieScene, const FGuid& ObjectHandle, TSubclassOf<class UMovieSceneTrack> TrackClass, FName UniqueTypeName) override;
+	virtual void BuildAddTrackMenu(FMenuBuilder& MenuBuilder) override;
+	virtual TSharedPtr<SWidget> BuildOutlinerEditWidget(const FGuid& ObjectBinding, UMovieSceneTrack* Track, const FBuildEditWidgetParams& Params) override;
+	virtual bool HandleAssetAdded(UObject* Asset, const FGuid& TargetObjectGuid) override;
+	virtual TSharedRef<ISequencerSection> MakeSectionInterface(UMovieSceneSection& SectionObject, UMovieSceneTrack& Track, FGuid ObjectBinding) override;
+	virtual bool SupportsSequence(UMovieSceneSequence* InSequence) const override;
+	virtual bool SupportsType(TSubclassOf<UMovieSceneTrack> TrackClass) const override;
+	virtual void Tick(float DeltaTime) override;
+	virtual const FSlateBrush* GetIconBrush() const override;
+
+protected:
+
+	/** Callback for AnimatablePropertyChanged in HandleAssetAdded for attached media sources. */
+	FKeyPropertyResult AddAttachedMediaSource(FFrameNumber KeyTime, class UMediaSource* MediaSource, TArray<TWeakObjectPtr<UObject>> ObjectsToAttachTo, int32 RowIndex);
+
+	/** Callback for AnimatablePropertyChanged in HandleAssetAdded for master media sources. */
+	FKeyPropertyResult AddMasterMediaSource(FFrameNumber KeyTime, class UMediaSource* MediaSource, int32 RowIndex);
+
+	void AddNewSection(const FAssetData& Asset, UMovieSceneMediaTrack* Track);
+
+	void AddNewSectionEnterPressed(const TArray<FAssetData>& Asset, UMovieSceneMediaTrack* Track);
+
+private:
+
+	/** Callback for executing the "Add Media Track" menu entry. */
+	void HandleAddMediaTrackMenuEntryExecute();
+
+private:
+
+	TSharedPtr<FTrackEditorThumbnailPool> ThumbnailPool;
+};
