Make some classes from MediaCompositing plugin public to they can inherited from in other plugins

diff --git a/Engine/Plugins/Media/MediaCompositing/Source/MediaCompositing/Public/MovieSceneMediaSection.h b/Engine/Plugins/Media/MediaCompositing/Source/MediaCompositing/Public/MovieSceneMediaSection.h
--- a/Engine/Plugins/Media/MediaCompositing/Source/MediaCompositing/Public/MovieSceneMediaSection.h
+++ b/Engine/Plugins/Media/MediaCompositing/Source/MediaCompositing/Public/MovieSceneMediaSection.h
@@ -15,8 +15,8 @@ class UMediaTexture;
 /**
  * Implements a movie scene section for media playback.
  */
-UCLASS(MinimalAPI)
-class UMovieSceneMediaSection
+UCLASS()
+class MEDIACOMPOSITING_API UMovieSceneMediaSection
 	: public UMovieSceneSection
 {
 public:
diff --git a/Engine/Plugins/Media/MediaCompositing/Source/MediaCompositing/Public/MovieSceneMediaTrack.h b/Engine/Plugins/Media/MediaCompositing/Source/MediaCompositing/Public/MovieSceneMediaTrack.h
--- a/Engine/Plugins/Media/MediaCompositing/Source/MediaCompositing/Public/MovieSceneMediaTrack.h
+++ b/Engine/Plugins/Media/MediaCompositing/Source/MediaCompositing/Public/MovieSceneMediaTrack.h
@@ -15,8 +15,8 @@ class UMediaSource;
 /**
  * Implements a movie scene track for media playback.
  */
-UCLASS(MinimalAPI)
-class UMovieSceneMediaTrack
+UCLASS()
+class MEDIACOMPOSITING_API UMovieSceneMediaTrack
 	: public UMovieSceneNameableTrack
 	, public IMovieSceneTrackTemplateProducer
 {
diff --git a/Engine/Plugins/Media/MediaCompositing/Source/MediaCompositingEditor/MediaCompositingEditor.Build.cs b/Engine/Plugins/Media/MediaCompositing/Source/MediaCompositingEditor/MediaCompositingEditor.Build.cs
--- a/Engine/Plugins/Media/MediaCompositing/Source/MediaCompositingEditor/MediaCompositingEditor.Build.cs
+++ b/Engine/Plugins/Media/MediaCompositing/Source/MediaCompositingEditor/MediaCompositingEditor.Build.cs
@@ -22,7 +22,6 @@ namespace UnrealBuildTool.Rules
 					"MediaCompositing",
 					"MediaUtils",
 					"MovieScene",
-					"MovieSceneTools",
 					"MovieSceneTracks",
 					"RHI",
 					"Sequencer",
@@ -33,6 +32,12 @@ namespace UnrealBuildTool.Rules
 					"TimeManagement"
 				});
 
+			PublicDependencyModuleNames.AddRange(
+				new string[]
+				{
+					"MovieSceneTools",
+				});
+
 			PrivateIncludePaths.AddRange(
 				new string[] {
 					"MediaCompositingEditor/Private",
diff --git a/Engine/Plugins/Media/MediaCompositing/Source/MediaCompositingEditor/Private/Sequencer/MediaPlayerPropertyTrackEditor.cpp b/Engine/Plugins/Media/MediaCompositing/Source/MediaCompositingEditor/Private/Sequencer/MediaPlayerPropertyTrackEditor.cpp
--- a/Engine/Plugins/Media/MediaCompositing/Source/MediaCompositingEditor/Private/Sequencer/MediaPlayerPropertyTrackEditor.cpp
+++ b/Engine/Plugins/Media/MediaCompositing/Source/MediaCompositingEditor/Private/Sequencer/MediaPlayerPropertyTrackEditor.cpp
@@ -17,7 +17,7 @@
 #include "MediaSource.h"
 #include "LevelSequence.h"
 
-#include "MediaThumbnailSection.h"
+#include "Sequencer/MediaThumbnailSection.h"
 
 #define LOCTEXT_NAMESPACE "FMediaPlayerPropertyTrackEditor"
 
diff --git a/Engine/Plugins/Media/MediaCompositing/Source/MediaCompositingEditor/Private/Sequencer/MediaThumbnailSection.cpp b/Engine/Plugins/Media/MediaCompositing/Source/MediaCompositingEditor/Private/Sequencer/MediaThumbnailSection.cpp
--- a/Engine/Plugins/Media/MediaCompositing/Source/MediaCompositingEditor/Private/Sequencer/MediaThumbnailSection.cpp
+++ b/Engine/Plugins/Media/MediaCompositing/Source/MediaCompositingEditor/Private/Sequencer/MediaThumbnailSection.cpp
@@ -1,6 +1,6 @@
 // Copyright Epic Games, Inc. All Rights Reserved.
 
-#include "MediaThumbnailSection.h"
+#include "Sequencer/MediaThumbnailSection.h"
 
 #include "EditorStyleSet.h"
 #include "IMediaCache.h"
@@ -35,7 +35,6 @@ FMediaThumbnailSection::FMediaThumbnailSection(UMovieSceneMediaSection& InSectio
 //	, MediaPlayer(nullptr)
 //	, MediaTexture(nullptr)
 	, SectionPtr(&InSection)
-	, SequencerPtr(InSequencer)
 {
 	TimeSpace = ETimeSpace::Local;
 }
diff --git a/Engine/Plugins/Media/MediaCompositing/Source/MediaCompositingEditor/Private/Sequencer/MediaThumbnailSection.h b/Engine/Plugins/Media/MediaCompositing/Source/MediaCompositingEditor/Private/Sequencer/MediaThumbnailSection.h
deleted file mode 100644
diff --git a/Engine/Plugins/Media/MediaCompositing/Source/MediaCompositingEditor/Private/Sequencer/MediaTrackEditor.cpp b/Engine/Plugins/Media/MediaCompositing/Source/MediaCompositingEditor/Private/Sequencer/MediaTrackEditor.cpp
--- a/Engine/Plugins/Media/MediaCompositing/Source/MediaCompositingEditor/Private/Sequencer/MediaTrackEditor.cpp
+++ b/Engine/Plugins/Media/MediaCompositing/Source/MediaCompositingEditor/Private/Sequencer/MediaTrackEditor.cpp
@@ -22,7 +22,7 @@
 #include "Misc/QualifiedFrameTime.h"
 #include "LevelSequence.h"
 
-#include "MediaThumbnailSection.h"
+#include "Sequencer/MediaThumbnailSection.h"
 
 
 #define LOCTEXT_NAMESPACE "FMediaTrackEditor"
diff --git a/Engine/Plugins/Media/MediaCompositing/Source/MediaCompositingEditor/Public/Sequencer/MediaThumbnailSection.h b/Engine/Plugins/Media/MediaCompositing/Source/MediaCompositingEditor/Public/Sequencer/MediaThumbnailSection.h
new file mode 100644
--- /dev/null
+++ b/Engine/Plugins/Media/MediaCompositing/Source/MediaCompositingEditor/Public/Sequencer/MediaThumbnailSection.h
@@ -0,0 +1,120 @@
+// Copyright Epic Games, Inc. All Rights Reserved.
+
+#pragma once
+
+#include "Math/RangeSet.h"
+#include "Misc/Timespan.h"
+#include "Sections/ThumbnailSection.h"
+#include "Templates/SharedPointer.h"
+#include "TrackEditorThumbnail/TrackEditorThumbnail.h"
+#include "UObject/GCObject.h"
+
+class FTrackEditorThumbnailPool;
+class ISequencer;
+class UMediaPlayer;
+class UMediaTexture;
+class UMovieSceneMediaSection;
+
+
+/**
+ * Implements a thumbnail section for media tracks.
+ */
+class MEDIACOMPOSITINGEDITOR_API FMediaThumbnailSection
+	: public FGCObject
+	, public FThumbnailSection
+	, public ICustomThumbnailClient
+{
+public:
+
+	/**
+	 * Create and initialize a new instance.
+	 *
+	 * @param InSection The movie scene section associated with this thumbnail section.
+	 * @param InThumbnailPool The thumbnail pool to use for drawing media frame thumbnails.
+	 * @param InSequencer The Sequencer object that owns this section.
+	 */
+	FMediaThumbnailSection(UMovieSceneMediaSection& InSection, TSharedPtr<FTrackEditorThumbnailPool> InThumbnailPool, TSharedPtr<ISequencer> InSequencer);
+
+	/** Virtual destructor. */
+	virtual ~FMediaThumbnailSection();
+
+public:
+
+	//~ FGCObject interface
+
+	virtual void AddReferencedObjects(FReferenceCollector& Collector) override;
+	virtual FString GetReferencerName() const override
+	{
+		return TEXT("FMediaThumbnailSection");
+	}
+
+public:
+
+	//~ FThumbnailSection interface
+
+	virtual FMargin GetContentPadding() const override;
+	virtual float GetSectionHeight() const override;
+	virtual FText GetSectionTitle() const override;
+	virtual void SetSingleTime(double GlobalTime) override;
+	virtual int32 OnPaintSection(FSequencerSectionPainter& InPainter) const override;
+	virtual void Tick(const FGeometry& AllottedGeometry, const FGeometry& ClippedGeometry, const double InCurrentTime, const float InDeltaTime) override;
+
+	// ISequencerSection interface
+
+	virtual void BeginResizeSection() override;
+	virtual void ResizeSection(ESequencerSectionResizeMode ResizeMode, FFrameNumber ResizeTime) override;
+	virtual void BeginSlipSection() override;
+	virtual void SlipSection(FFrameNumber SlipTime) override;
+
+public:
+
+	//~ ICustomThumbnailClient interface
+
+	virtual void Draw(FTrackEditorThumbnail& TrackEditorThumbnail) override;
+	virtual void Setup() override;
+
+protected:
+
+	/**
+	 * Draw the section's film border decoration.
+	 *
+	 * @param InPainter The object that paints the geometry.
+	 * @param SectionSize The size of the section (in local coordinates).
+	 */
+	void DrawFilmBorder(FSequencerSectionPainter& InPainter, FVector2D SectionSize) const;
+
+	/**
+	 * Draw indicators for where the media source is looping.
+	 *
+	 * @param InPainter The object that paints the geometry.
+	 * @param SectionSize The size of the section (in local coordinates).
+	 */
+	void DrawLoopIndicators(FSequencerSectionPainter& InPainter, FTimespan MediaDuration, FVector2D SectionSize) const;
+
+	/** Draw the caching state of the given media samples. */
+	void DrawSampleStates(FSequencerSectionPainter& InPainter, FTimespan MediaDuration, FVector2D SectionSize, const TRangeSet<FTimespan>& RangeSet, const FLinearColor& Color) const;
+
+	/**
+	 * Get the media player that is used by the evaluation template.
+	 *
+	 * @return The media player, or nullptr if not found.
+	 */
+	UMediaPlayer* GetTemplateMediaPlayer() const;
+
+private:
+
+	/** Internal media player used to generate the thumbnail images. */
+//	UMediaPlayer* MediaPlayer;
+
+	/** Media texture that receives the thumbnail image frames. */
+//	UMediaTexture* MediaTexture;
+
+	/** The section object that owns this section. */
+	TWeakObjectPtr<UMovieSceneMediaSection> SectionPtr;
+
+	/** Cached start offset value valid only during resize */
+	FFrameNumber InitialStartOffsetDuringResize;
+
+	/** Cached start time valid only during resize */
+	FFrameNumber InitialStartTimeDuringResize;
+};
diff --git a/Engine/Plugins/Media/MediaCompositing/Source/MediaCompositingEditor/Public/Sequencer/MediaTrackEditor.h b/Engine/Plugins/Media/MediaCompositing/Source/MediaCompositingEditor/Public/Sequencer/MediaTrackEditor.h
--- a/Engine/Plugins/Media/MediaCompositing/Source/MediaCompositingEditor/Public/Sequencer/MediaTrackEditor.h
+++ b/Engine/Plugins/Media/MediaCompositing/Source/MediaCompositingEditor/Public/Sequencer/MediaTrackEditor.h
@@ -77,6 +77,7 @@ protected:
 
 	void AddNewSectionEnterPressed(const TArray<FAssetData>& Asset, UMovieSceneMediaTrack* Track);
 
+	TSharedPtr<FTrackEditorThumbnailPool> GetThumbnailPool() const { return ThumbnailPool; }
 private:
 
 	/** Callback for executing the "Add Media Track" menu entry. */
