Remove commented out code from a few MediaCompositing classes

diff --git a/Engine/Plugins/Media/MediaCompositing/Source/MediaCompositingEditor/Private/Sequencer/MediaThumbnailSection.cpp b/Engine/Plugins/Media/MediaCompositing/Source/MediaCompositingEditor/Private/Sequencer/MediaThumbnailSection.cpp
--- a/Engine/Plugins/Media/MediaCompositing/Source/MediaCompositingEditor/Private/Sequencer/MediaThumbnailSection.cpp
+++ b/Engine/Plugins/Media/MediaCompositing/Source/MediaCompositingEditor/Private/Sequencer/MediaThumbnailSection.cpp
@@ -32,8 +32,6 @@
 
 FMediaThumbnailSection::FMediaThumbnailSection(UMovieSceneMediaSection& InSection, TSharedPtr<FTrackEditorThumbnailPool> InThumbnailPool, TSharedPtr<ISequencer> InSequencer)
 	: FThumbnailSection(InSequencer, InThumbnailPool, this, InSection)
-//	, MediaPlayer(nullptr)
-//	, MediaTexture(nullptr)
 	, SectionPtr(&InSection)
 {
 	TimeSpace = ETimeSpace::Local;
@@ -42,10 +40,6 @@ FMediaThumbnailSection::FMediaThumbnailSection(UMovieSceneMediaSection& InSectio
 
 FMediaThumbnailSection::~FMediaThumbnailSection()
 {
-//	if (MediaPlayer != nullptr)
-//	{
-//		MediaPlayer->Close();
-//	}
 }
 
 
@@ -54,8 +48,6 @@ FMediaThumbnailSection::~FMediaThumbnailSection()
 
 void FMediaThumbnailSection::AddReferencedObjects(FReferenceCollector& Collector)
 {
-//	Collector.AddReferencedObject(MediaPlayer);
-//	Collector.AddReferencedObject(MediaTexture);
 }
 
 
@@ -246,81 +238,12 @@ void FMediaThumbnailSection::SlipSection(FFrameNumber SlipTime)
  *****************************************************************************/
 
 void FMediaThumbnailSection::Draw(FTrackEditorThumbnail& TrackEditorThumbnail)
-{/*
-	check(MediaPlayer != nullptr);
-
-	if (MediaPlayer->IsPreparing())
-	{
-		RedrawThumbnails();
-
-		return;
-	}
-
-	check(MediaTexture != nullptr);
-	check(MediaTexture->Resource != nullptr);
-	check(MediaTexture->Resource->TextureRHI.IsValid());
-
-	// get target texture resource
-	FTexture2DRHIRef Texture2DRHI = MediaTexture->Resource->TextureRHI->GetTexture2D();
-
-	if (!Texture2DRHI.IsValid())
-	{
-		return;
-	}
-
-	// seek media player to thumbnail position
-	const float EvalPosition = FMath::Max(0.0f, TrackEditorThumbnail.GetEvalPosition());
-	const FTimespan EvalTime = int64(EvalPosition * ETimespan::TicksPerSecond);
-	const FTimespan MediaTime = EvalTime % MediaPlayer->GetDuration();
-
-	if (!MediaPlayer->Seek(FTimespan(MediaTime)))
-	{
-		return;
-	}
-
-	// resolve media player texture to track editor thumbnail
-	TrackEditorThumbnail.CopyTextureIn(Texture2DRHI);
-
-	TSharedPtr<ISequencer> Sequencer = SequencerPtr.Pin();
-
-	if (Sequencer.IsValid())
-	{
-		TrackEditorThumbnail.SetupFade(Sequencer->GetSequencerWidget());
-	}*/
+{
 }
 
 
 void FMediaThumbnailSection::Setup()
-{/*
-	UMovieSceneMediaSection* MediaSection = CastChecked<UMovieSceneMediaSection>(Section);
-	UMediaSource* MediaSource = MediaSection->GetMediaSource();
-
-	if (MediaSource == nullptr)
-	{
-		return;
-	}
-
-	// create internal player
-	if (MediaPlayer == nullptr)
-	{
-		MediaPlayer = NewObject<UMediaPlayer>(GetTransientPackage(), MakeUniqueObjectName(GetTransientPackage(), UMediaPlayer::StaticClass()));
-	}
-
-	// create target texture
-	if (MediaTexture == nullptr)
-	{
-		MediaTexture = NewObject<UMediaTexture>(GetTransientPackage(), MakeUniqueObjectName(GetTransientPackage(), UMediaTexture::StaticClass()));
-		MediaTexture->SetMediaPlayer(MediaPlayer);
-		MediaTexture->UpdateResource();
-	}
-
-	// open latest media source
-	if (MediaPlayer->GetUrl() != MediaSource->GetUrl())
-	{
-		MediaPlayer->OpenSource(MediaSource);
-	}
-
-	MediaPlayer->Pause();*/
+{
 }
 
 
@@ -353,28 +276,6 @@ void FMediaThumbnailSection::DrawFilmBorder(FSequencerSectionPainter& InPainter,
 
 void FMediaThumbnailSection::DrawLoopIndicators(FSequencerSectionPainter& InPainter, FTimespan MediaDuration, FVector2D SectionSize) const
 {
-	/*
-	static const FSlateBrush* GenericBrush = FCoreStyle::Get().GetBrush("GenericWhiteBox");
-
-	FFrameRate TickResolution = Section->GetTypedOuter<UMovieScene>()->GetTickResolution();
-	double SectionDuration = FFrameTime(UE::MovieScene::DiscreteSize(Section->GetRange())) / TickResolution;
-	const float MediaSizeX = MediaDuration.GetTotalSeconds() * SectionSize.X / SectionDuration;
-	float DrawOffset = 0.0f;
-
-	while (DrawOffset < SectionSize.X)
-	{
-		FSlateDrawElement::MakeBox(
-			InPainter.DrawElements,
-			InPainter.LayerId++,
-			InPainter.SectionGeometry.ToPaintGeometry(FVector2D(DrawOffset, 0.0f), FVector2D(1.0f, SectionSize.Y)),
-			GenericBrush,
-			ESlateDrawEffect::None,
-			FLinearColor::Gray
-		);
-
-		DrawOffset += MediaSizeX;
-	}
-	*/
 	static const FSlateBrush* GenericBrush = FCoreStyle::Get().GetBrush("GenericWhiteBox");
 
 	UMovieSceneMediaSection* MediaSection = Cast<UMovieSceneMediaSection>(Section);
diff --git a/Engine/Plugins/Media/MediaCompositing/Source/MediaCompositingEditor/Public/Sequencer/MediaThumbnailSection.h b/Engine/Plugins/Media/MediaCompositing/Source/MediaCompositingEditor/Public/Sequencer/MediaThumbnailSection.h
--- a/Engine/Plugins/Media/MediaCompositing/Source/MediaCompositingEditor/Public/Sequencer/MediaThumbnailSection.h
+++ b/Engine/Plugins/Media/MediaCompositing/Source/MediaCompositingEditor/Public/Sequencer/MediaThumbnailSection.h
@@ -103,12 +103,6 @@ protected:
 
 private:
 
-	/** Internal media player used to generate the thumbnail images. */
-//	UMediaPlayer* MediaPlayer;
-
-	/** Media texture that receives the thumbnail image frames. */
-//	UMediaTexture* MediaTexture;
-
 	/** The section object that owns this section. */
 	TWeakObjectPtr<UMovieSceneMediaSection> SectionPtr;
 
