MediaCompositing: Fixed media cache display being offset when track is offset from the start.
#rnx

diff --git a/Engine/Plugins/Media/MediaCompositing/Source/MediaCompositingEditor/Private/Sequencer/MediaThumbnailSection.cpp b/Engine/Plugins/Media/MediaCompositing/Source/MediaCompositingEditor/Private/Sequencer/MediaThumbnailSection.cpp
--- a/Engine/Plugins/Media/MediaCompositing/Source/MediaCompositingEditor/Private/Sequencer/MediaThumbnailSection.cpp
+++ b/Engine/Plugins/Media/MediaCompositing/Source/MediaCompositingEditor/Private/Sequencer/MediaThumbnailSection.cpp
@@ -320,7 +320,9 @@ void FMediaThumbnailSection::DrawSampleStates(FSequencerSectionPainter& InPainte
 
 	for (auto& Range : Ranges)
 	{
-		const float DrawOffset = FMath::RoundToNegativeInfinity(FTimespan::Ratio(Range.GetLowerBoundValue(), MediaDuration) * MediaSizeX) - TimeToPixelConverter.SecondsToPixel(TickResolution.AsSeconds(MediaSection->StartFrameOffset));
+		const float DrawOffset = FMath::RoundToNegativeInfinity(FTimespan::Ratio(Range.GetLowerBoundValue(), MediaDuration) * MediaSizeX) -
+			TimeToPixelConverter.SecondsToPixel(TickResolution.AsSeconds(MediaSection->StartFrameOffset)) +
+			TimeToPixelConverter.SecondsToPixel(0.0);
 		const float DrawSize = FMath::RoundToPositiveInfinity(FTimespan::Ratio(Range.Size<FTimespan>(), MediaDuration) * MediaSizeX);
 		const float BarHeight = 4.0f;
 
