Sequencer: Media track can be added to media plate object bindings.

#rnx

diff --git a/Engine/Plugins/Media/MediaCompositing/Source/MediaCompositingEditor/MediaCompositingEditor.Build.cs b/Engine/Plugins/Media/MediaCompositing/Source/MediaCompositingEditor/MediaCompositingEditor.Build.cs
--- a/Engine/Plugins/Media/MediaCompositing/Source/MediaCompositingEditor/MediaCompositingEditor.Build.cs
+++ b/Engine/Plugins/Media/MediaCompositing/Source/MediaCompositingEditor/MediaCompositingEditor.Build.cs
@@ -44,6 +44,15 @@ namespace UnrealBuildTool.Rules
 					"MediaCompositingEditor/Private/Sequencer",
 					"MediaCompositingEditor/Private/Shared",
 				});
+
+			DynamicallyLoadedModuleNames.AddRange(
+				new string[] {
+					"MediaPlate",
+				});
+			PrivateIncludePathModuleNames.AddRange(
+				new string[] {
+					"MediaPlate",
+				});
 		}
 	}
 }
diff --git a/Engine/Plugins/Media/MediaCompositing/Source/MediaCompositingEditor/Private/Sequencer/MediaTrackEditor.cpp b/Engine/Plugins/Media/MediaCompositing/Source/MediaCompositingEditor/Private/Sequencer/MediaTrackEditor.cpp
--- a/Engine/Plugins/Media/MediaCompositing/Source/MediaCompositingEditor/Private/Sequencer/MediaTrackEditor.cpp
+++ b/Engine/Plugins/Media/MediaCompositing/Source/MediaCompositingEditor/Private/Sequencer/MediaTrackEditor.cpp
@@ -11,6 +11,8 @@
 #include "ISequencerObjectChangeListener.h"
 #include "MediaSource.h"
 #include "MediaTexture.h"
+#include "MediaPlate.h"
+#include "MediaPlateModule.h"
 #include "MovieSceneBinding.h"
 #include "MovieSceneMediaSection.h"
 #include "MovieSceneMediaTrack.h"
@@ -76,6 +78,27 @@ void FMediaTrackEditor::BuildAddTrackMenu(FMenuBuilder& MenuBuilder)
 	);
 }
 
+void FMediaTrackEditor::BuildObjectBindingTrackMenu(FMenuBuilder& MenuBuilder, const TArray<FGuid>& ObjectBindings, const UClass* ObjectClass)
+{
+	// Is this a media plate?
+	if (ObjectClass != nullptr)
+	{
+		FMediaPlateModule* MediaPlateModule = FModuleManager::LoadModulePtr<FMediaPlateModule>("MediaPlate");
+		if (MediaPlateModule != nullptr)
+		{
+			if (ObjectClass->IsChildOf(MediaPlateModule->GetAMediaPlateClass()))
+			{
+				MenuBuilder.AddMenuEntry(
+					LOCTEXT("AddTrack", "Media"),
+					LOCTEXT("AddAttachedTooltip", "Adds a media track attached to the object."),
+					FSlateIcon(),
+					FUIAction(
+						FExecuteAction::CreateRaw(this, &FMediaTrackEditor::HandleAddMediaTrackToObjectBindingMenuEntryExecute, ObjectBindings)));
+			}
+		}
+	}
+
+}
 
 TSharedPtr<SWidget> FMediaTrackEditor::BuildOutlinerEditWidget(const FGuid& ObjectBinding, UMovieSceneTrack* Track, const FBuildEditWidgetParams& Params)
 {
@@ -311,6 +334,38 @@ void FMediaTrackEditor::HandleAddMediaTrackMenuEntryExecute()
 	}
 }
 
+void FMediaTrackEditor::HandleAddMediaTrackToObjectBindingMenuEntryExecute(TArray<FGuid> InObjectBindingIDs)
+{
+	UMovieScene* FocusedMovieScene = GetFocusedMovieScene();
+
+	if (FocusedMovieScene == nullptr)
+	{
+		return;
+	}
 
+	if (FocusedMovieScene->IsReadOnly())
+	{
+		return;
+	}
+
+	const FScopedTransaction Transaction(NSLOCTEXT("Sequencer", "AddMediaTrack_Transaction", "Add Media Track"));
+	FocusedMovieScene->Modify();
+
+	// Loop through all objects.
+	for (FGuid InObjectBindingID : InObjectBindingIDs)
+	{
+		if (InObjectBindingID.IsValid())
+		{
+			// Add media track.
+			UMovieSceneMediaTrack* NewObjectTrack = FocusedMovieScene->AddTrack<UMovieSceneMediaTrack>(InObjectBindingID);
+			NewObjectTrack->SetDisplayName(LOCTEXT("MediaTrackName", "Media"));
+			
+			if (GetSequencer().IsValid())
+			{
+				GetSequencer()->OnAddTrack(NewObjectTrack, InObjectBindingID);
+			}
+		}
+	}
+}
 
 #undef LOCTEXT_NAMESPACE
diff --git a/Engine/Plugins/Media/MediaCompositing/Source/MediaCompositingEditor/Public/Sequencer/MediaTrackEditor.h b/Engine/Plugins/Media/MediaCompositing/Source/MediaCompositingEditor/Public/Sequencer/MediaTrackEditor.h
--- a/Engine/Plugins/Media/MediaCompositing/Source/MediaCompositingEditor/Public/Sequencer/MediaTrackEditor.h
+++ b/Engine/Plugins/Media/MediaCompositing/Source/MediaCompositingEditor/Public/Sequencer/MediaTrackEditor.h
@@ -57,6 +57,7 @@ public:
 
 	virtual UMovieSceneTrack* AddTrack(UMovieScene* FocusedMovieScene, const FGuid& ObjectHandle, TSubclassOf<class UMovieSceneTrack> TrackClass, FName UniqueTypeName) override;
 	virtual void BuildAddTrackMenu(FMenuBuilder& MenuBuilder) override;
+	virtual void BuildObjectBindingTrackMenu(FMenuBuilder& MenuBuilder, const TArray<FGuid>& ObjectBindings, const UClass* ObjectClass) override;
 	virtual TSharedPtr<SWidget> BuildOutlinerEditWidget(const FGuid& ObjectBinding, UMovieSceneTrack* Track, const FBuildEditWidgetParams& Params) override;
 	virtual bool HandleAssetAdded(UObject* Asset, const FGuid& TargetObjectGuid) override;
 	virtual TSharedRef<ISequencerSection> MakeSectionInterface(UMovieSceneSection& SectionObject, UMovieSceneTrack& Track, FGuid ObjectBinding) override;
@@ -82,6 +83,8 @@ private:
 
 	/** Callback for executing the "Add Media Track" menu entry. */
 	void HandleAddMediaTrackMenuEntryExecute();
+	/** Callback for adding a media track to an object binding. */
+	void HandleAddMediaTrackToObjectBindingMenuEntryExecute(TArray<FGuid> InObjectBindingID);
 
 private:
 
