Sequencer: Media track can use the media player from a media plate.

#rnx

diff --git a/Engine/Plugins/Media/MediaCompositing/Source/MediaCompositing/MediaCompositing.Build.cs b/Engine/Plugins/Media/MediaCompositing/Source/MediaCompositing/MediaCompositing.Build.cs
--- a/Engine/Plugins/Media/MediaCompositing/Source/MediaCompositing/MediaCompositing.Build.cs
+++ b/Engine/Plugins/Media/MediaCompositing/Source/MediaCompositing/MediaCompositing.Build.cs
@@ -32,6 +32,15 @@ namespace UnrealBuildTool.Rules
 					"MediaCompositing/Private",
 					"MediaCompositing/Private/MovieScene",
 				});
+
+			DynamicallyLoadedModuleNames.AddRange(
+				new string[] {
+					"MediaPlate",
+				});
+			PrivateIncludePathModuleNames.AddRange(
+				new string[] {
+					"MediaPlate",
+				});
 		}
 	}
 }
diff --git a/Engine/Plugins/Media/MediaCompositing/Source/MediaCompositing/Private/MovieScene/MovieSceneMediaTemplate.cpp b/Engine/Plugins/Media/MediaCompositing/Source/MediaCompositing/Private/MovieScene/MovieSceneMediaTemplate.cpp
--- a/Engine/Plugins/Media/MediaCompositing/Source/MediaCompositing/Private/MovieScene/MovieSceneMediaTemplate.cpp
+++ b/Engine/Plugins/Media/MediaCompositing/Source/MediaCompositing/Private/MovieScene/MovieSceneMediaTemplate.cpp
@@ -3,6 +3,7 @@
 #include "MovieSceneMediaTemplate.h"
 
 #include "Math/UnrealMathUtility.h"
+#include "MediaPlateModule.h"
 #include "MediaPlayer.h"
 #include "MediaPlayerFacade.h"
 #include "MediaSoundComponent.h"
@@ -22,8 +23,39 @@
 /* Local helpers
  *****************************************************************************/
 
-struct FMediaSectionPreRollExecutionToken
+ /** Base struct for media execution tokens. */
+struct FBaseMediaSectionExecutionToken
 	: IMovieSceneExecutionToken
+{
+	/**
+	 * Call this to get a media player from the bound objects.
+	 * Only one object is supported at the moment.
+	 */
+	UMediaPlayer* GetMediaPlayer(const FMovieSceneEvaluationOperand& Operand, IMovieScenePlayer& Player)
+	{
+		UMediaPlayer* MediaPlayer = nullptr;
+		if (Operand.ObjectBindingID.IsValid())
+		{
+			FMediaPlateModule* MediaPlateModule = FModuleManager::LoadModulePtr<FMediaPlateModule>("MediaPlate");
+			if (MediaPlateModule != nullptr)
+			{
+				for (TWeakObjectPtr<> WeakObject : Player.FindBoundObjects(Operand))
+				{
+					UObject* BoundObject = WeakObject.Get();
+					if (BoundObject != nullptr)
+					{
+						MediaPlayer = MediaPlateModule->GetMediaPlayer(BoundObject);
+						break;
+					}
+				}
+			}
+		}
+		return MediaPlayer;
+	}
+};
+
+struct FMediaSectionPreRollExecutionToken
+	: FBaseMediaSectionExecutionToken
 {
 	FMediaSectionPreRollExecutionToken(UMediaSource* InMediaSource, FTimespan InStartTimeSeconds)
 		: MediaSource(InMediaSource)
@@ -35,7 +67,11 @@ struct FMediaSectionPreRollExecutionToken
 		using namespace PropertyTemplate;
 
 		FMovieSceneMediaData& SectionData = PersistentData.GetSectionData<FMovieSceneMediaData>();
-		UMediaPlayer* MediaPlayer = SectionData.GetMediaPlayer();
+		UMediaPlayer* MediaPlayer = GetMediaPlayer(Operand, Player);
+		if (MediaPlayer == nullptr)
+		{
+			MediaPlayer = SectionData.GetMediaPlayer();
+		}
 
 		if (MediaPlayer == nullptr || MediaSource == nullptr)
 		{
@@ -58,7 +94,7 @@ private:
 
 
 struct FMediaSectionExecutionToken
-	: IMovieSceneExecutionToken
+	: FBaseMediaSectionExecutionToken
 {
 	FMediaSectionExecutionToken(UMediaSource* InMediaSource, FTimespan InCurrentTime, FTimespan InFrameDuration)
 		: CurrentTime(InCurrentTime)
@@ -70,7 +106,11 @@ struct FMediaSectionExecutionToken
 	virtual void Execute(const FMovieSceneContext& Context, const FMovieSceneEvaluationOperand& Operand, FPersistentEvaluationData& PersistentData, IMovieScenePlayer& Player) override
 	{
 		FMovieSceneMediaData& SectionData = PersistentData.GetSectionData<FMovieSceneMediaData>();
-		UMediaPlayer* MediaPlayer = SectionData.GetMediaPlayer();
+		UMediaPlayer* MediaPlayer = GetMediaPlayer(Operand, Player);
+		if (MediaPlayer == nullptr)
+		{
+			MediaPlayer = SectionData.GetMediaPlayer();
+		}
 
 		if (MediaPlayer == nullptr || MediaSource == nullptr)
 		{
