Sequencer: Media track section data setup moved to Initialize so that we can skip making a media player when we have a media plate.

#rnx

diff --git a/Engine/Plugins/Media/MediaCompositing/Source/MediaCompositing/Private/MovieScene/MovieSceneMediaTemplate.cpp b/Engine/Plugins/Media/MediaCompositing/Source/MediaCompositing/Private/MovieScene/MovieSceneMediaTemplate.cpp
--- a/Engine/Plugins/Media/MediaCompositing/Source/MediaCompositing/Private/MovieScene/MovieSceneMediaTemplate.cpp
+++ b/Engine/Plugins/Media/MediaCompositing/Source/MediaCompositing/Private/MovieScene/MovieSceneMediaTemplate.cpp
@@ -23,39 +23,8 @@
 /* Local helpers
  *****************************************************************************/
 
- /** Base struct for media execution tokens. */
-struct FBaseMediaSectionExecutionToken
-	: IMovieSceneExecutionToken
-{
-	/**
-	 * Call this to get a media player from the bound objects.
-	 * Only one object is supported at the moment.
-	 */
-	UMediaPlayer* GetMediaPlayer(const FMovieSceneEvaluationOperand& Operand, IMovieScenePlayer& Player)
-	{
-		UMediaPlayer* MediaPlayer = nullptr;
-		if (Operand.ObjectBindingID.IsValid())
-		{
-			FMediaPlateModule* MediaPlateModule = FModuleManager::LoadModulePtr<FMediaPlateModule>("MediaPlate");
-			if (MediaPlateModule != nullptr)
-			{
-				for (TWeakObjectPtr<> WeakObject : Player.FindBoundObjects(Operand))
-				{
-					UObject* BoundObject = WeakObject.Get();
-					if (BoundObject != nullptr)
-					{
-						MediaPlayer = MediaPlateModule->GetMediaPlayer(BoundObject);
-						break;
-					}
-				}
-			}
-		}
-		return MediaPlayer;
-	}
-};
-
 struct FMediaSectionPreRollExecutionToken
-	: FBaseMediaSectionExecutionToken
+	: IMovieSceneExecutionToken
 {
 	FMediaSectionPreRollExecutionToken(UMediaSource* InMediaSource, FTimespan InStartTimeSeconds)
 		: MediaSource(InMediaSource)
@@ -67,11 +36,7 @@ struct FMediaSectionPreRollExecutionToken
 		using namespace PropertyTemplate;
 
 		FMovieSceneMediaData& SectionData = PersistentData.GetSectionData<FMovieSceneMediaData>();
-		UMediaPlayer* MediaPlayer = GetMediaPlayer(Operand, Player);
-		if (MediaPlayer == nullptr)
-		{
-			MediaPlayer = SectionData.GetMediaPlayer();
-		}
+		UMediaPlayer* MediaPlayer = SectionData.GetMediaPlayer();
 
 		if (MediaPlayer == nullptr || MediaSource == nullptr)
 		{
@@ -94,7 +59,7 @@ private:
 
 
 struct FMediaSectionExecutionToken
-	: FBaseMediaSectionExecutionToken
+	: IMovieSceneExecutionToken
 {
 	FMediaSectionExecutionToken(UMediaSource* InMediaSource, FTimespan InCurrentTime, FTimespan InFrameDuration)
 		: CurrentTime(InCurrentTime)
@@ -106,11 +71,7 @@ struct FMediaSectionExecutionToken
 	virtual void Execute(const FMovieSceneContext& Context, const FMovieSceneEvaluationOperand& Operand, FPersistentEvaluationData& PersistentData, IMovieScenePlayer& Player) override
 	{
 		FMovieSceneMediaData& SectionData = PersistentData.GetSectionData<FMovieSceneMediaData>();
-		UMediaPlayer* MediaPlayer = GetMediaPlayer(Operand, Player);
-		if (MediaPlayer == nullptr)
-		{
-			MediaPlayer = SectionData.GetMediaPlayer();
-		}
+		UMediaPlayer* MediaPlayer = SectionData.GetMediaPlayer();
 
 		if (MediaPlayer == nullptr || MediaSource == nullptr)
 		{
@@ -316,6 +277,36 @@ UScriptStruct& FMovieSceneMediaSectionTemplate::GetScriptStructImpl() const
 void FMovieSceneMediaSectionTemplate::Initialize(const FMovieSceneEvaluationOperand& Operand, const FMovieSceneContext& Context, FPersistentEvaluationData& PersistentData, IMovieScenePlayer& Player) const
 {
 	FMovieSceneMediaData* SectionData = PersistentData.FindSectionData<FMovieSceneMediaData>();
+	if (SectionData == nullptr)
+	{
+		// Are we overriding the media player?
+		UMediaPlayer* MediaPlayer = Params.MediaPlayer;
+		if (MediaPlayer == nullptr)
+		{
+			// Nope... do we have an object binding?
+			if (Operand.ObjectBindingID.IsValid())
+			{
+				// Yes. Get the media player from the object.
+				FMediaPlateModule* MediaPlateModule = FModuleManager::LoadModulePtr<FMediaPlateModule>("MediaPlate");
+				if (MediaPlateModule != nullptr)
+				{
+					for (TWeakObjectPtr<> WeakObject : Player.FindBoundObjects(Operand))
+					{
+						UObject* BoundObject = WeakObject.Get();
+						if (BoundObject != nullptr)
+						{
+							MediaPlayer = MediaPlateModule->GetMediaPlayer(BoundObject);
+							break;
+						}
+					}
+				}
+			}
+		}
+
+		// Add section data.
+		SectionData = &PersistentData.AddSectionData<FMovieSceneMediaData>();
+		SectionData->Setup(MediaPlayer);
+	}
 
 	if (!ensure(SectionData != nullptr))
 	{
@@ -375,15 +366,9 @@ void FMovieSceneMediaSectionTemplate::Initialize(const FMovieSceneEvaluationOper
 }
 
 
-void FMovieSceneMediaSectionTemplate::Setup(FPersistentEvaluationData& PersistentData, IMovieScenePlayer& Player) const
-{
-	PersistentData.AddSectionData<FMovieSceneMediaData>().Setup(Params.MediaPlayer);
-}
-
-
 void FMovieSceneMediaSectionTemplate::SetupOverrides()
 {
-	EnableOverrides(RequiresInitializeFlag | RequiresSetupFlag | RequiresTearDownFlag);
+	EnableOverrides(RequiresInitializeFlag | RequiresTearDownFlag);
 }
 
 
diff --git a/Engine/Plugins/Media/MediaCompositing/Source/MediaCompositing/Private/MovieScene/MovieSceneMediaTemplate.h b/Engine/Plugins/Media/MediaCompositing/Source/MediaCompositing/Private/MovieScene/MovieSceneMediaTemplate.h
--- a/Engine/Plugins/Media/MediaCompositing/Source/MediaCompositing/Private/MovieScene/MovieSceneMediaTemplate.h
+++ b/Engine/Plugins/Media/MediaCompositing/Source/MediaCompositing/Private/MovieScene/MovieSceneMediaTemplate.h
@@ -78,7 +78,6 @@ public:
 	virtual void Evaluate(const FMovieSceneEvaluationOperand& Operand, const FMovieSceneContext& Context, const FPersistentEvaluationData& PersistentData, FMovieSceneExecutionTokens& ExecutionTokens) const override;
 	virtual UScriptStruct& GetScriptStructImpl() const override;
 	virtual void Initialize(const FMovieSceneEvaluationOperand& Operand, const FMovieSceneContext& Context, FPersistentEvaluationData& PersistentData, IMovieScenePlayer& Player) const override;
-	virtual void Setup(FPersistentEvaluationData& PersistentData, IMovieScenePlayer& Player) const;
 	virtual void SetupOverrides() override;
 	virtual void TearDown(FPersistentEvaluationData& PersistentData, IMovieScenePlayer& Player) const override;
 
