MediaCompositing: Fixed crash when adding a media track as a default for an actor.
#rnx

diff --git a/Engine/Plugins/Media/MediaCompositing/Source/MediaCompositing/Private/MovieScene/MovieSceneMediaTrack.cpp b/Engine/Plugins/Media/MediaCompositing/Source/MediaCompositing/Private/MovieScene/MovieSceneMediaTrack.cpp
--- a/Engine/Plugins/Media/MediaCompositing/Source/MediaCompositing/Private/MovieScene/MovieSceneMediaTrack.cpp
+++ b/Engine/Plugins/Media/MediaCompositing/Source/MediaCompositing/Private/MovieScene/MovieSceneMediaTrack.cpp
@@ -27,6 +27,7 @@ UMovieSceneMediaTrack::UMovieSceneMediaTrack(const FObjectInitializer& ObjectIni
 
 #if WITH_EDITORONLY_DATA
 	TrackTint = FColor(0, 0, 0, 200);
+	bSupportsDefaultSections = false;
 #endif
 }
 
