Sequencer: Added cache settings to the media track.

#rnx

diff --git a/Engine/Plugins/Media/MediaCompositing/Source/MediaCompositing/Private/MovieScene/MovieSceneMediaTemplate.cpp b/Engine/Plugins/Media/MediaCompositing/Source/MediaCompositing/Private/MovieScene/MovieSceneMediaTemplate.cpp
--- a/Engine/Plugins/Media/MediaCompositing/Source/MediaCompositing/Private/MovieScene/MovieSceneMediaTemplate.cpp
+++ b/Engine/Plugins/Media/MediaCompositing/Source/MediaCompositing/Private/MovieScene/MovieSceneMediaTemplate.cpp
@@ -203,6 +203,10 @@ FMovieSceneMediaSectionTemplate::FMovieSceneMediaSectionTemplate(const UMovieSce
 	Params.MediaSoundComponent = InSection.MediaSoundComponent;
 	Params.bLooping = InSection.bLooping;
 	Params.StartFrameOffset = InSection.StartFrameOffset;
+	if (Params.MediaSource != nullptr)
+	{
+		Params.MediaSource->SetCacheSettings(InSection.CacheSettings);
+	}
 
 	// If using an external media player link it here so we don't automatically create it later.
 	Params.MediaPlayer = InSection.bUseExternalMediaPlayer ? InSection.ExternalMediaPlayer : nullptr;
diff --git a/Engine/Plugins/Media/MediaCompositing/Source/MediaCompositing/Public/MovieSceneMediaSection.h b/Engine/Plugins/Media/MediaCompositing/Source/MediaCompositing/Public/MovieSceneMediaSection.h
--- a/Engine/Plugins/Media/MediaCompositing/Source/MediaCompositing/Public/MovieSceneMediaSection.h
+++ b/Engine/Plugins/Media/MediaCompositing/Source/MediaCompositing/Public/MovieSceneMediaSection.h
@@ -4,6 +4,8 @@
 
 #include "MovieSceneSection.h"
 
+#include "MediaSource.h"
+
 #include "MovieSceneMediaSection.generated.h"
 
 class UMediaPlayer;
@@ -53,6 +55,10 @@ public:
 	UPROPERTY(EditAnywhere, BlueprintReadWrite, Category="Media", meta = (EditCondition = "bUseExternalMediaPlayer"), AdvancedDisplay)
 	TObjectPtr<UMediaPlayer> ExternalMediaPlayer;
 
+	/** Override the default cache settings. */
+	UPROPERTY(EditAnywhere, BlueprintReadWrite, Category = "Media", AdvancedDisplay)
+	FMediaSourceCacheSettings CacheSettings;
+
 public:
 
 	/**
