MediaCompositing: Fixed media tracks not sending the full frame duration to Mediaframework.

#rnx

diff --git a/Engine/Plugins/Media/MediaCompositing/Source/MediaCompositing/Private/MovieScene/MovieSceneMediaTemplate.cpp b/Engine/Plugins/Media/MediaCompositing/Source/MediaCompositing/Private/MovieScene/MovieSceneMediaTemplate.cpp
--- a/Engine/Plugins/Media/MediaCompositing/Source/MediaCompositing/Private/MovieScene/MovieSceneMediaTemplate.cpp
+++ b/Engine/Plugins/Media/MediaCompositing/Source/MediaCompositing/Private/MovieScene/MovieSceneMediaTemplate.cpp
@@ -253,7 +253,7 @@ void FMovieSceneMediaSectionTemplate::Evaluate(const FMovieSceneEvaluationOperan
 		const double FrameTimeInSeconds = FrameRate.AsSeconds(FrameTime);
 		const int64 FrameTicks = FrameTimeInSeconds * ETimespan::TicksPerSecond;
 
-		const double FrameDurationInSeconds = FrameRate.AsSeconds(FFrameTime(1));
+		const double FrameDurationInSeconds = (Context.GetRange().Size<FFrameTime>()) / Context.GetFrameRate();
 		const int64 FrameDurationTicks = FrameDurationInSeconds * ETimespan::TicksPerSecond;
 
 		#if MOVIESCENEMEDIATEMPLATE_TRACE_EVALUATION
