MediaCompositing: Media track displays tile info if available.
#rnx

diff --git a/Engine/Plugins/Media/MediaCompositing/Source/MediaCompositingEditor/Private/Sequencer/MediaThumbnailSection.cpp b/Engine/Plugins/Media/MediaCompositing/Source/MediaCompositingEditor/Private/Sequencer/MediaThumbnailSection.cpp
--- a/Engine/Plugins/Media/MediaCompositing/Source/MediaCompositingEditor/Private/Sequencer/MediaThumbnailSection.cpp
+++ b/Engine/Plugins/Media/MediaCompositing/Source/MediaCompositingEditor/Private/Sequencer/MediaThumbnailSection.cpp
@@ -2,6 +2,7 @@
 
 #include "Sequencer/MediaThumbnailSection.h"
 
+#include "Fonts/FontMeasure.h"
 #include "Styling/AppStyle.h"
 #include "IMediaCache.h"
 #include "IMediaTracks.h"
@@ -138,6 +139,8 @@ int32 FMediaThumbnailSection::OnPaintSection(FSequencerSectionPainter& InPainter
 		DrawSampleStates(InPainter, MediaDuration, SectionSize, CacheRangeSet, FLinearColor(0.07059, 0.32941, 0.07059));
 
 		DrawLoopIndicators(InPainter, MediaDuration, SectionSize);
+
+		DrawMediaInfo(InPainter, MediaPlayer, SectionSize);
 	}
 	InPainter.DrawElements.PopClip();
 
@@ -343,6 +346,42 @@ void FMediaThumbnailSection::DrawSampleStates(FSequencerSectionPainter& InPainte
 }
 
 
+void FMediaThumbnailSection::DrawMediaInfo(FSequencerSectionPainter& InPainter, 
+	UMediaPlayer* MediaPlayer,FVector2D SectionSize) const
+{
+	// Get tile info.
+	FString TileString;
+	FIntPoint TileNum = MediaPlayer->GetTileNum();
+	int32 TileTotalNum = TileNum.X * TileNum.Y;
+	if (TileTotalNum > 1)
+	{
+		TileString = FText::Format(LOCTEXT("TileNum", "Tiles: {0}"), TileTotalNum).ToString();
+	}
+
+	const ESlateDrawEffect DrawEffects = InPainter.bParentEnabled ? ESlateDrawEffect::None : ESlateDrawEffect::DisabledEffect;
+	const FSlateFontInfo SmallLayoutFont = FCoreStyle::GetDefaultFontStyle("Bold", 10);
+	const TSharedRef< FSlateFontMeasure > FontMeasureService = FSlateApplication::Get().GetRenderer()->GetFontMeasureService();
+	FVector2D TextSize = FontMeasureService->Measure(TileString, SmallLayoutFont);
+
+	FMargin ContentPadding = GetContentPadding();
+	FVector2D TextOffset(ContentPadding.Left + 4, 
+		InPainter.SectionGeometry.Size.Y - (TextSize.Y + ContentPadding.Bottom));
+
+	// Add tile string.
+	if (TileString.Len() > 0)
+	{
+		FSlateDrawElement::MakeText(
+			InPainter.DrawElements,
+			InPainter.LayerId++,
+			InPainter.SectionGeometry.ToPaintGeometry(TextOffset, TextSize),
+			TileString,
+			SmallLayoutFont,
+			DrawEffects,
+			FColor(192, 192, 192, static_cast<uint8>(InPainter.GhostAlpha * 255))
+		);
+	}
+}
+
 UMediaPlayer* FMediaThumbnailSection::GetTemplateMediaPlayer() const
 {
 	// locate the track that evaluates this section
diff --git a/Engine/Plugins/Media/MediaCompositing/Source/MediaCompositingEditor/Public/Sequencer/MediaThumbnailSection.h b/Engine/Plugins/Media/MediaCompositing/Source/MediaCompositingEditor/Public/Sequencer/MediaThumbnailSection.h
--- a/Engine/Plugins/Media/MediaCompositing/Source/MediaCompositingEditor/Public/Sequencer/MediaThumbnailSection.h
+++ b/Engine/Plugins/Media/MediaCompositing/Source/MediaCompositingEditor/Public/Sequencer/MediaThumbnailSection.h
@@ -94,6 +94,15 @@ protected:
 	/** Draw the caching state of the given media samples. */
 	void DrawSampleStates(FSequencerSectionPainter& InPainter, FTimespan MediaDuration, FVector2D SectionSize, const TRangeSet<FTimespan>& RangeSet, const FLinearColor& Color) const;
 
+	/**
+	 * Draw the info about the current media.
+	 *
+	 * @param InPainter			The object that paints the geometry.
+	 * @param MediaPlayer		The player to get media info from.
+	 * @param SectionSize		The size of the section (in local coordinates).
+	 */
+	void DrawMediaInfo(FSequencerSectionPainter& InPainter, UMediaPlayer* MediaPlayer, FVector2D SectionSize) const;
+
 	/**
 	 * Get the media player that is used by the evaluation template.
 	 *
