MediaAssets: Added MediaPlayerProxyInterface to coordinate player control between the owning object and an external user.
#rnx

diff --git a/Engine/Plugins/Media/MediaCompositing/Source/MediaCompositing/Private/MovieScene/MovieSceneMediaData.cpp b/Engine/Plugins/Media/MediaCompositing/Source/MediaCompositing/Private/MovieScene/MovieSceneMediaData.cpp
--- a/Engine/Plugins/Media/MediaCompositing/Source/MediaCompositing/Private/MovieScene/MovieSceneMediaData.cpp
+++ b/Engine/Plugins/Media/MediaCompositing/Source/MediaCompositing/Private/MovieScene/MovieSceneMediaData.cpp
@@ -4,6 +4,7 @@
 
 #include "IMediaEventSink.h"
 #include "MediaPlayer.h"
+#include "MediaPlayerProxyInterface.h"
 #include "UObject/Package.h"
 #include "UObject/UObjectGlobals.h"
 
@@ -44,7 +45,7 @@ void FMovieSceneMediaData::SeekOnOpen(FTimespan Time)
 }
 
 
-void FMovieSceneMediaData::Setup(UMediaPlayer* OverrideMediaPlayer)
+void FMovieSceneMediaData::Setup(UMediaPlayer* OverrideMediaPlayer, UObject* InPlayerProxy)
 {
 	// Ensure we don't already have a media player set. Setup should only be called once
 	check(!MediaPlayer);
@@ -62,6 +63,16 @@ void FMovieSceneMediaData::Setup(UMediaPlayer* OverrideMediaPlayer)
 	MediaPlayer->PlayOnOpen = false;
 	MediaPlayer->OnMediaEvent().AddRaw(this, &FMovieSceneMediaData::HandleMediaPlayerEvent);
 	MediaPlayer->AddToRoot();
+
+	// Do we have a valid proxy object?
+	if ((InPlayerProxy != nullptr) && (InPlayerProxy->Implements<UMediaPlayerProxyInterface>()))
+	{
+		PlayerProxy = InPlayerProxy;
+	}
+	else
+	{
+		PlayerProxy.Reset();
+	}
 }
 
 
diff --git a/Engine/Plugins/Media/MediaCompositing/Source/MediaCompositing/Private/MovieScene/MovieSceneMediaTemplate.cpp b/Engine/Plugins/Media/MediaCompositing/Source/MediaCompositing/Private/MovieScene/MovieSceneMediaTemplate.cpp
--- a/Engine/Plugins/Media/MediaCompositing/Source/MediaCompositing/Private/MovieScene/MovieSceneMediaTemplate.cpp
+++ b/Engine/Plugins/Media/MediaCompositing/Source/MediaCompositing/Private/MovieScene/MovieSceneMediaTemplate.cpp
@@ -6,6 +6,7 @@
 #include "Math/UnrealMathUtility.h"
 #include "MediaPlayer.h"
 #include "MediaPlayerFacade.h"
+#include "MediaPlayerProxyInterface.h"
 #include "MediaSoundComponent.h"
 #include "MediaSource.h"
 #include "MediaTexture.h"
@@ -72,12 +73,28 @@ struct FMediaSectionExecutionToken
 	{
 		FMovieSceneMediaData& SectionData = PersistentData.GetSectionData<FMovieSceneMediaData>();
 		UMediaPlayer* MediaPlayer = SectionData.GetMediaPlayer();
+		UObject* PlayerProxy = SectionData.GetPlayerProxy();
 
 		if (MediaPlayer == nullptr || MediaSource == nullptr)
 		{
 			return;
 		}
 
+		// Do we have a player proxy?
+		if (PlayerProxy != nullptr)
+		{
+			IMediaPlayerProxyInterface* PlayerProxyInterface = Cast<IMediaPlayerProxyInterface>
+				(PlayerProxy);
+			if (PlayerProxyInterface != nullptr)
+			{
+				// Can we control the player?
+				if (PlayerProxyInterface->IsExternalControlAllowed() == false)
+				{
+					return;
+				}
+			}
+		}
+
 		// open the media source if necessary
 		if (MediaPlayer->GetUrl().IsEmpty())
 		{
@@ -285,6 +302,7 @@ void FMovieSceneMediaSectionTemplate::Initialize(const FMovieSceneEvaluationOper
 	{
 		// Are we overriding the media player?
 		UMediaPlayer* MediaPlayer = Params.MediaPlayer;
+		UObject* PlayerProxy = nullptr;
 		if (MediaPlayer == nullptr)
 		{
 			// Nope... do we have an object binding?
@@ -299,7 +317,7 @@ void FMovieSceneMediaSectionTemplate::Initialize(const FMovieSceneEvaluationOper
 						UObject* BoundObject = WeakObject.Get();
 						if (BoundObject != nullptr)
 						{
-							MediaPlayer = MediaAssetsModule->GetPlayerFromObject(BoundObject);
+							MediaPlayer = MediaAssetsModule->GetPlayerFromObject(BoundObject, PlayerProxy);
 							break;
 						}
 					}
@@ -309,7 +327,7 @@ void FMovieSceneMediaSectionTemplate::Initialize(const FMovieSceneEvaluationOper
 
 		// Add section data.
 		SectionData = &PersistentData.AddSectionData<FMovieSceneMediaData>();
-		SectionData->Setup(MediaPlayer);
+		SectionData->Setup(MediaPlayer, PlayerProxy);
 	}
 
 	if (!ensure(SectionData != nullptr))
diff --git a/Engine/Plugins/Media/MediaCompositing/Source/MediaCompositing/Public/MovieSceneMediaData.h b/Engine/Plugins/Media/MediaCompositing/Source/MediaCompositing/Public/MovieSceneMediaData.h
--- a/Engine/Plugins/Media/MediaCompositing/Source/MediaCompositing/Public/MovieSceneMediaData.h
+++ b/Engine/Plugins/Media/MediaCompositing/Source/MediaCompositing/Public/MovieSceneMediaData.h
@@ -4,7 +4,9 @@
 
 #include "Misc/Timespan.h"
 #include "Evaluation/MovieScenePropertyTemplate.h"
+#include "UObject/WeakObjectPtrTemplates.h"
 
+class IMediaPlayerProxy;
 class UMediaPlayer;
 enum class EMediaEvent;
 
@@ -30,6 +32,11 @@ public:
 	 */
 	UMediaPlayer* GetMediaPlayer();
 
+	/**
+	 * Get the optional proxy object used by this persistent data.
+	 */
+	UObject* GetPlayerProxy() { return PlayerProxy.Get(); }
+
 	/**
 	 * Set the time to seek to after opening a media source has finished.
 	 *
@@ -38,7 +45,7 @@ public:
 	void SeekOnOpen(FTimespan Time);
 
 	/** Set up this persistent data object. */
-	void Setup(UMediaPlayer* OverrideMediaPlayer);
+	void Setup(UMediaPlayer* OverrideMediaPlayer, UObject* InPlayerProxy);
 
 private:
 
@@ -50,6 +57,8 @@ private:
 
 	/** The media player used by this object. */
 	UMediaPlayer* MediaPlayer;
+	/** Optional proxy for the media player. */
+	TWeakObjectPtr<UObject> PlayerProxy;
 
 	/** The time to seek to after the media source is opened. */
 	FTimespan SeekOnOpenTime;
diff --git a/Engine/Plugins/Media/MediaPlate/Source/MediaPlate/Private/MediaPlateComponent.cpp b/Engine/Plugins/Media/MediaPlate/Source/MediaPlate/Private/MediaPlateComponent.cpp
--- a/Engine/Plugins/Media/MediaPlate/Source/MediaPlate/Private/MediaPlateComponent.cpp
+++ b/Engine/Plugins/Media/MediaPlate/Source/MediaPlate/Private/MediaPlateComponent.cpp
@@ -153,12 +153,11 @@ void UMediaPlateComponent::TickComponent(float DeltaTime, enum ELevelTick TickTy
 {
 	Super::TickComponent(DeltaTime, TickType, ThisTickFunction);
 
-	if (bPlayOnlyWhenVisible)
+	if ((MediaPlayer != nullptr) && (bPlayOnlyWhenVisible))
 	{
-		bool bIsVisible = IsVisible();
-
-		if (MediaPlayer != nullptr)
+		if (MediaPlayer->GetRate() != 0.0f)
 		{
+			bool bIsVisible = IsVisible();
 			if (bIsVisible)
 			{
 				ResumeWhenVisible();
@@ -444,6 +443,12 @@ void UMediaPlateComponent::TickOutput()
 	}
 }
 
+bool UMediaPlateComponent::IsExternalControlAllowed()
+{
+	// Allow control if we always play, or are visible.
+	return (bPlayOnlyWhenVisible == false) || IsVisible();
+}
+
 void UMediaPlateComponent::RestartPlayer()
 {
 	if (MediaPlayer != nullptr)
diff --git a/Engine/Plugins/Media/MediaPlate/Source/MediaPlate/Private/MediaPlateModule.cpp b/Engine/Plugins/Media/MediaPlate/Source/MediaPlate/Private/MediaPlateModule.cpp
--- a/Engine/Plugins/Media/MediaPlate/Source/MediaPlate/Private/MediaPlateModule.cpp
+++ b/Engine/Plugins/Media/MediaPlate/Source/MediaPlate/Private/MediaPlateModule.cpp
@@ -8,7 +8,7 @@
 
 DEFINE_LOG_CATEGORY(LogMediaPlate);
 
-UMediaPlayer* FMediaPlateModule::GetMediaPlayer(UObject* Object)
+UMediaPlayer* FMediaPlateModule::GetMediaPlayer(UObject* Object, UObject*& PlayerProxy)
 {
 	UMediaPlayer* MediaPlayer = nullptr;
 
@@ -20,6 +20,7 @@ UMediaPlayer* FMediaPlateModule::GetMediaPlayer(UObject* Object)
 		if (MediaPlateComponent != nullptr)
 		{
 			MediaPlayer = MediaPlateComponent->GetMediaPlayer();
+			PlayerProxy = MediaPlateComponent;
 		}
 	}
 
diff --git a/Engine/Plugins/Media/MediaPlate/Source/MediaPlate/Public/MediaPlateComponent.h b/Engine/Plugins/Media/MediaPlate/Source/MediaPlate/Public/MediaPlateComponent.h
--- a/Engine/Plugins/Media/MediaPlate/Source/MediaPlate/Public/MediaPlateComponent.h
+++ b/Engine/Plugins/Media/MediaPlate/Source/MediaPlate/Public/MediaPlateComponent.h
@@ -5,6 +5,7 @@
 #include "Components/ActorComponent.h"
 #include "CoreMinimal.h"
 #include "MediaPlayerOptions.h"
+#include "MediaPlayerProxyInterface.h"
 #include "MediaSource.h"
 #include "MediaTextureTracker.h"
 
@@ -23,7 +24,8 @@ class UMediaTexture;
  * This is a component for AMediaPlate that can play and show media in the world.
  */
 UCLASS()
-class MEDIAPLATE_API UMediaPlateComponent : public UActorComponent
+class MEDIAPLATE_API UMediaPlateComponent : public UActorComponent,
+	public IMediaPlayerProxyInterface
 {
 	GENERATED_UCLASS_BODY()
 
@@ -214,6 +216,9 @@ public:
 	 */
 	void TickOutput();
 
+	//~ IMediaPlayerProxyInterface.
+	virtual bool IsExternalControlAllowed() override;
+
 private:
 	void RestartPlayer();
 
diff --git a/Engine/Plugins/Media/MediaPlate/Source/MediaPlate/Public/MediaPlateModule.h b/Engine/Plugins/Media/MediaPlate/Source/MediaPlate/Public/MediaPlateModule.h
--- a/Engine/Plugins/Media/MediaPlate/Source/MediaPlate/Public/MediaPlateModule.h
+++ b/Engine/Plugins/Media/MediaPlate/Source/MediaPlate/Public/MediaPlateModule.h
@@ -7,6 +7,7 @@
 #include "Modules/ModuleManager.h"
 
 class AMediaPlate;
+class IMediaPlayerProxy;
 class UMaterialInterface;
 class UMediaPlayer;
 class UObject;
@@ -27,7 +28,7 @@ public:
 	/**
 	 * Call this to get the media player from a media plate object.
 	 */
-	UMediaPlayer* GetMediaPlayer(UObject* Object);
+	UMediaPlayer* GetMediaPlayer(UObject* Object, UObject*& PlayerProxy);
 
 	/** IModuleInterface implementation */
 	virtual void StartupModule() override;
diff --git a/Engine/Source/Runtime/MediaAssets/Private/Assets/MediaPlayerProxyInterface.cpp b/Engine/Source/Runtime/MediaAssets/Private/Assets/MediaPlayerProxyInterface.cpp
new file mode 100644
--- /dev/null
+++ b/Engine/Source/Runtime/MediaAssets/Private/Assets/MediaPlayerProxyInterface.cpp
@@ -0,0 +1,3 @@
+// Copyright Epic Games, Inc. All Rights Reserved.
+
+#include "MediaPlayerProxyInterface.h"
diff --git a/Engine/Source/Runtime/MediaAssets/Private/MediaAssetsModule.cpp b/Engine/Source/Runtime/MediaAssets/Private/MediaAssetsModule.cpp
--- a/Engine/Source/Runtime/MediaAssets/Private/MediaAssetsModule.cpp
+++ b/Engine/Source/Runtime/MediaAssets/Private/MediaAssetsModule.cpp
@@ -69,16 +69,17 @@ public:
 		GetPlayerFromObjectDelegates[DelegateID].Unbind();
 	}
 
-	virtual UMediaPlayer* GetPlayerFromObject(UObject* Object) override
+	virtual UMediaPlayer* GetPlayerFromObject(UObject* Object, UObject*& PlayerProxy) override
 	{
 		UMediaPlayer* Player = nullptr;
+		PlayerProxy = nullptr;
 
 		// Go through all our delegates.
 		for (FOnGetPlayerFromObject& Delegate : GetPlayerFromObjectDelegates)
 		{
 			if (Delegate.IsBound())
 			{
-				Player = Delegate.Execute(Object);
+				Player = Delegate.Execute(Object, PlayerProxy);
 				if (Player != nullptr)
 				{
 					break;
diff --git a/Engine/Source/Runtime/MediaAssets/Public/IMediaAssetsModule.h b/Engine/Source/Runtime/MediaAssets/Public/IMediaAssetsModule.h
--- a/Engine/Source/Runtime/MediaAssets/Public/IMediaAssetsModule.h
+++ b/Engine/Source/Runtime/MediaAssets/Public/IMediaAssetsModule.h
@@ -16,7 +16,7 @@ class IMediaAssetsModule
 {
 public:
 	/** Delegate to get a player from a UObject. */
-	DECLARE_DELEGATE_RetVal_OneParam(UMediaPlayer*, FOnGetPlayerFromObject, UObject*);
+	DECLARE_DELEGATE_RetVal_TwoParams(UMediaPlayer*, FOnGetPlayerFromObject, UObject*, UObject*& /*PlayerProxy*/);
 
 	/**
 	 * Plugins should call this so they can provide a function to get a media player from an object.
@@ -34,13 +34,16 @@ public:
 	virtual void UnregisterGetPlayerFromObject(int32 DelegateID) = 0;
 
 	/**
-	 * Call this to get a media player from an object.
+	 * Call this to get a media player (and proxy object) from an object.
 	 * This will query any plugins that have called RegisterGetPlayerFromObject.
 	 * 
-	 * @param Object	Object to get the player from.
+	 * The proxy object will implement IMediaPlayerProxyInterface.
+	 * 
+	 * @param Object		Object to get the player from.
+	 * @param PlayerProxy	Will be set to the proxy object (or nullptr if none).
 	 * @return Media player, or nullptr if none found. 
 	 */
-	virtual UMediaPlayer* GetPlayerFromObject(UObject* Object) = 0;
+	virtual UMediaPlayer* GetPlayerFromObject(UObject* Object, UObject*& PlayerProxy) = 0;
 
 	/** Virtual destructor. */
 	virtual ~IMediaAssetsModule() { }
diff --git a/Engine/Source/Runtime/MediaAssets/Public/MediaPlayerProxyInterface.h b/Engine/Source/Runtime/MediaAssets/Public/MediaPlayerProxyInterface.h
new file mode 100644
--- /dev/null
+++ b/Engine/Source/Runtime/MediaAssets/Public/MediaPlayerProxyInterface.h
@@ -0,0 +1,31 @@
+// Copyright Epic Games, Inc. All Rights Reserved.
+
+#pragma once
+
+#include "UObject/Interface.h"
+
+#include "MediaPlayerProxyInterface.generated.h"
+
+/**
+ * The proxy object provides a higher level of insight/control than just the media player.
+ * For example, the object owning the player may control the player in some
+ * cases, and the proxy allows you and the object to avoid conflicts in control.
+ */
+UINTERFACE(MinimalAPI)
+class UMediaPlayerProxyInterface : public UInterface
+{
+	GENERATED_BODY()
+};
+
+class MEDIAASSETS_API IMediaPlayerProxyInterface
+{
+	GENERATED_BODY()
+
+public:
+	/**
+	 * Call this to see if you can control the media player, or if the owning object is using it.
+	 * 
+	 * @return				True if you can control the player.
+	 */
+	virtual bool IsExternalControlAllowed() = 0;
+};
