MediaCompositing: Media track will show number of mip levels if available.
#rnx

diff --git a/Engine/Plugins/Media/ImgMedia/Source/ImgMedia/Private/Player/ImgMediaPlayer.cpp b/Engine/Plugins/Media/ImgMedia/Source/ImgMedia/Private/Player/ImgMediaPlayer.cpp
--- a/Engine/Plugins/Media/ImgMedia/Source/ImgMedia/Private/Player/ImgMediaPlayer.cpp
+++ b/Engine/Plugins/Media/ImgMedia/Source/ImgMedia/Private/Player/ImgMediaPlayer.cpp
@@ -161,8 +161,17 @@ FVariant FImgMediaPlayer::GetMediaInfo(FName InfoName) const
 {
 	FVariant Variant;
 
+	// Source num mips?
+	if (InfoName == UMediaPlayer::MediaInfoNameSourceNumMips.Resolve())
+	{
+		if (Loader.IsValid())
+		{
+			int32 NumMips = Loader->GetNumMipLevels();
+			Variant = NumMips;
+		}
+	}
 	// Source num tiles?
-	if (InfoName == UMediaPlayer::MediaInfoNameSourceNumTiles.Resolve())
+	else if (InfoName == UMediaPlayer::MediaInfoNameSourceNumTiles.Resolve())
 	{
 		if (Loader.IsValid())
 		{
diff --git a/Engine/Plugins/Media/MediaCompositing/Source/MediaCompositingEditor/Private/Sequencer/MediaThumbnailSection.cpp b/Engine/Plugins/Media/MediaCompositing/Source/MediaCompositingEditor/Private/Sequencer/MediaThumbnailSection.cpp
--- a/Engine/Plugins/Media/MediaCompositing/Source/MediaCompositingEditor/Private/Sequencer/MediaThumbnailSection.cpp
+++ b/Engine/Plugins/Media/MediaCompositing/Source/MediaCompositingEditor/Private/Sequencer/MediaThumbnailSection.cpp
@@ -359,18 +359,31 @@ void FMediaThumbnailSection::DrawMediaInfo(FSequencerSectionPainter& InPainter,
 		TileString = FText::Format(LOCTEXT("TileNum", "Tiles: {0}"), TileTotalNum).ToString();
 	}
 
+	// Get mip info.
+	FString MipString;
+	int32 MipNum = 0;
+	MediaPlayer->GetMediaInfo<int32>(MipNum, UMediaPlayer::MediaInfoNameSourceNumMips.Resolve());
+	if (MipNum > 1)
+	{
+		MipString = FText::Format(LOCTEXT("Mips", "Mips: {0}"), MipNum).ToString();
+	}
+
 	const ESlateDrawEffect DrawEffects = InPainter.bParentEnabled ? ESlateDrawEffect::None : ESlateDrawEffect::DisabledEffect;
 	const FSlateFontInfo SmallLayoutFont = FCoreStyle::GetDefaultFontStyle("Bold", 10);
 	const TSharedRef< FSlateFontMeasure > FontMeasureService = FSlateApplication::Get().GetRenderer()->GetFontMeasureService();
-	FVector2D TextSize = FontMeasureService->Measure(TileString, SmallLayoutFont);
-
+	
+	FVector2D TextSize(EForceInit::ForceInitToZero);
 	FMargin ContentPadding = GetContentPadding();
-	FVector2D TextOffset(ContentPadding.Left + 4, 
-		InPainter.SectionGeometry.Size.Y - (TextSize.Y + ContentPadding.Bottom));
+	FVector2D TextOffset(EForceInit::ForceInitToZero);
+	float BaseYOffset = 0.0f;
 
 	// Add tile string.
 	if (TileString.Len() > 0)
 	{
+		TextSize = FontMeasureService->Measure(TileString, SmallLayoutFont);
+		TextOffset.Set(ContentPadding.Left + 4,
+			InPainter.SectionGeometry.Size.Y - (TextSize.Y + ContentPadding.Bottom));
+		TextOffset.Y += BaseYOffset;
 		FSlateDrawElement::MakeText(
 			InPainter.DrawElements,
 			InPainter.LayerId++,
@@ -380,6 +393,26 @@ void FMediaThumbnailSection::DrawMediaInfo(FSequencerSectionPainter& InPainter,
 			DrawEffects,
 			FColor(192, 192, 192, static_cast<uint8>(InPainter.GhostAlpha * 255))
 		);
+		BaseYOffset -= TextSize.Y + 2.0f;
+	}
+
+	// Add mips string.
+	if (MipString.Len() > 0)
+	{
+		TextSize = FontMeasureService->Measure(MipString, SmallLayoutFont);
+		TextOffset.Set(ContentPadding.Left + 4,
+			InPainter.SectionGeometry.Size.Y - (TextSize.Y + ContentPadding.Bottom));
+		TextOffset.Y += BaseYOffset;
+		FSlateDrawElement::MakeText(
+			InPainter.DrawElements,
+			InPainter.LayerId++,
+			InPainter.SectionGeometry.ToPaintGeometry(TextOffset, TextSize),
+			MipString,
+			SmallLayoutFont,
+			DrawEffects,
+			FColor(192, 192, 192, static_cast<uint8>(InPainter.GhostAlpha * 255))
+		);
+		BaseYOffset -= TextSize.Y + 2.0f;
 	}
 }
 
