MediaPlate: Fixed PlayOnlyWhenVisible not resuming.
#rnx

diff --git a/Engine/Plugins/Media/MediaCompositing/Source/MediaCompositing/Private/MovieScene/MovieSceneMediaTemplate.cpp b/Engine/Plugins/Media/MediaCompositing/Source/MediaCompositing/Private/MovieScene/MovieSceneMediaTemplate.cpp
--- a/Engine/Plugins/Media/MediaCompositing/Source/MediaCompositing/Private/MovieScene/MovieSceneMediaTemplate.cpp
+++ b/Engine/Plugins/Media/MediaCompositing/Source/MediaCompositing/Private/MovieScene/MovieSceneMediaTemplate.cpp
@@ -81,9 +81,10 @@ struct FMediaSectionExecutionToken
 		}
 
 		// Do we have a player proxy?
+		IMediaPlayerProxyInterface* PlayerProxyInterface = nullptr;
 		if (PlayerProxy != nullptr)
 		{
-			IMediaPlayerProxyInterface* PlayerProxyInterface = Cast<IMediaPlayerProxyInterface>
+			PlayerProxyInterface = Cast<IMediaPlayerProxyInterface>
 				(PlayerProxy);
 			if (PlayerProxyInterface != nullptr)
 			{
@@ -152,7 +153,8 @@ struct FMediaSectionExecutionToken
 				// Set rate
 				// (note that the DIRECTION is important, but the magnitude is not - as we use blocked playback, the range setup to block on will serve as external clock to the player,
 				//  the direction is taken into account as hint for internal operation of the player)
-				if (!MediaPlayer->SetRate((Context.GetDirection() == EPlayDirection::Forwards) ? 1.0f : -1.0f))
+				if (!SetRate(MediaPlayer, PlayerProxyInterface,
+					(Context.GetDirection() == EPlayDirection::Forwards) ? 1.0f : -1.0f))
 				{
 					// Failed to set needed rate: better switch off blocking and bail...
 					MediaPlayer->SetBlockOnTimeRange(TRange<FTimespan>::Empty());
@@ -169,7 +171,7 @@ struct FMediaSectionExecutionToken
 				float CurrentPlayerRate = MediaPlayer->GetRate();
 				if (Context.GetDirection() == EPlayDirection::Forwards && CurrentPlayerRate < 0.0f)
 				{
-					if (!MediaPlayer->SetRate(1.0f))
+					if (!SetRate(MediaPlayer, PlayerProxyInterface, 1.0f))
 					{
 						// Failed to set needed rate: better switch off blocking and bail...
 						MediaPlayer->SetBlockOnTimeRange(TRange<FTimespan>::Empty());
@@ -178,7 +180,7 @@ struct FMediaSectionExecutionToken
 				}
 				else if (Context.GetDirection() == EPlayDirection::Backwards && CurrentPlayerRate > 0.0f)
 				{
-					if (!MediaPlayer->SetRate(-1.0f))
+					if (!SetRate(MediaPlayer, PlayerProxyInterface, -1.0f))
 					{
 						// Failed to set needed rate: better switch off blocking and bail...
 						MediaPlayer->SetBlockOnTimeRange(TRange<FTimespan>::Empty());
@@ -191,7 +193,7 @@ struct FMediaSectionExecutionToken
 		{
 			if (MediaPlayer->IsPlaying())
 			{
-				MediaPlayer->SetRate(0.0f);
+				SetRate(MediaPlayer, PlayerProxyInterface, 0.0f);
 			}
 
 			MediaPlayer->Seek(MediaTime);
@@ -204,6 +206,22 @@ struct FMediaSectionExecutionToken
 
 private:
 
+	/**
+	 * Call this to set the rate.
+	 * It will use the proxy if available, otherwise it will use the player.
+	 */
+	bool SetRate(UMediaPlayer* Player, IMediaPlayerProxyInterface* Proxy, float Rate)
+	{
+		if (Proxy != nullptr)
+		{
+			return Proxy->SetProxyRate(Rate);
+		}
+		else
+		{
+			return Player->SetRate(Rate);
+		}
+	}
+
 	FTimespan CurrentTime;
 	FTimespan FrameDuration;
 	UMediaSource* MediaSource;
diff --git a/Engine/Plugins/Media/MediaPlate/Source/MediaPlate/Private/MediaPlateComponent.cpp b/Engine/Plugins/Media/MediaPlate/Source/MediaPlate/Private/MediaPlateComponent.cpp
--- a/Engine/Plugins/Media/MediaPlate/Source/MediaPlate/Private/MediaPlateComponent.cpp
+++ b/Engine/Plugins/Media/MediaPlate/Source/MediaPlate/Private/MediaPlateComponent.cpp
@@ -155,7 +155,7 @@ void UMediaPlateComponent::TickComponent(float DeltaTime, enum ELevelTick TickTy
 
 	if ((MediaPlayer != nullptr) && (bPlayOnlyWhenVisible))
 	{
-		if (MediaPlayer->GetRate() != 0.0f)
+		if ((CurrentRate != 0.0f) || (bWantsToPlayWhenVisible))
 		{
 			bool bIsVisible = IsVisible();
 			if (bIsVisible)
@@ -186,6 +186,8 @@ UMediaTexture* UMediaPlateComponent::GetMediaTexture()
 
 void UMediaPlateComponent::Open()
 {
+	bIsMediaPlatePlaying = true;
+	CurrentRate = bPlayOnOpen ? 1.0f : 0.0f;
 	if ((bPlayOnlyWhenVisible == false) || (IsVisible()))
 	{
 		bool bIsPlaying = false;
@@ -218,6 +220,7 @@ void UMediaPlateComponent::Play()
 	{
 		MediaPlayer->Play();
 	}
+	CurrentRate = 1.0f;
 }
 
 void UMediaPlateComponent::Pause()
@@ -226,6 +229,7 @@ void UMediaPlateComponent::Pause()
 	{
 		MediaPlayer->Pause();
 	}
+	CurrentRate = 0.0f;
 }
 
 bool UMediaPlateComponent::Rewind()
@@ -251,6 +255,7 @@ void UMediaPlateComponent::Close()
 	}
 
 	StopClockSink();
+	bIsMediaPlatePlaying = false;
 	bWantsToPlayWhenVisible = false;
 	bResumeWhenOpened = false;
 	PlaylistIndex = 0;
@@ -443,6 +448,24 @@ void UMediaPlateComponent::TickOutput()
 	}
 }
 
+float UMediaPlateComponent::GetProxyRate() const
+{
+	return CurrentRate;
+}
+
+bool UMediaPlateComponent::SetProxyRate(float Rate)
+{
+	CurrentRate = Rate;
+
+	bool bSuccess = true;
+	if (MediaPlayer != nullptr)
+	{
+		bSuccess = MediaPlayer->SetRate(Rate);
+	}
+
+	return bSuccess;
+}
+
 bool UMediaPlateComponent::IsExternalControlAllowed()
 {
 	// Allow control if we always play, or are visible.
@@ -495,6 +518,7 @@ void UMediaPlateComponent::ResumeWhenVisible()
 				(MediaPlayer->IsPlaying() == false))
 			{
 				bResumeWhenOpened = true;
+				bWantsToPlayWhenVisible = false;
 				Open();
 			}
 		}
@@ -522,7 +546,8 @@ FTimespan UMediaPlateComponent::GetResumeTime()
 			}
 			else
 			{
-				PlayerTime = MediaDuration - FTimespan(1);
+				// It wont play if we seek to the very end, so go back a little bit.
+				PlayerTime = MediaDuration - FTimespan::FromSeconds(0.001f);
 			}
 		}
 	}
@@ -632,6 +657,7 @@ void UMediaPlateComponent::OnMediaOpened(FString DeviceUrl)
 	if (bResumeWhenOpened)
 	{
 		bResumeWhenOpened = false;
+		bIsMediaPlatePlaying = true;
 		if (MediaPlayer != nullptr)
 		{
 			FTimespan PlayTime = GetResumeTime();
diff --git a/Engine/Plugins/Media/MediaPlate/Source/MediaPlate/Public/MediaPlateComponent.h b/Engine/Plugins/Media/MediaPlate/Source/MediaPlate/Public/MediaPlateComponent.h
--- a/Engine/Plugins/Media/MediaPlate/Source/MediaPlate/Public/MediaPlateComponent.h
+++ b/Engine/Plugins/Media/MediaPlate/Source/MediaPlate/Public/MediaPlateComponent.h
@@ -97,6 +97,12 @@ public:
 	UFUNCTION(BlueprintCallable, Category = "Media|MediaPlateComponent")
 	void Close();
 
+	/**
+	 * Call this to see if the media plate is playing.
+	 */
+	UFUNCTION(BlueprintGetter)
+	bool IsMediaPlatePlaying() const { return bIsMediaPlatePlaying; }
+
 	/**
 	 * Call this to see if we want to loop.
 	 */
@@ -217,11 +223,24 @@ public:
 	void TickOutput();
 
 	//~ IMediaPlayerProxyInterface.
+	virtual float GetProxyRate() const override;
+	virtual bool SetProxyRate(float Rate) override;
 	virtual bool IsExternalControlAllowed() override;
 
 private:
 	void RestartPlayer();
 
+	/**
+	 * If true, then we want the media plate to play.
+	 * Note that this could be true, but the player is not actually playing because
+	 * bPlayOnlyWhenVisible = true and the plate is not visible.
+	 */
+	UPROPERTY(Blueprintgetter = IsMediaPlatePlaying, Category = "MediaPlate", meta = (AllowPrivateAccess = true))
+	bool bIsMediaPlatePlaying = false;
+
+	/** Desired rate of play that we want. */
+	float CurrentRate = 0.0f;
+
 	/** If true then only allow playback when the media plate is visible. */
 	UPROPERTY(EditAnywhere, BlueprintReadOnly, Category = "MediaPlate", meta = (AllowPrivateAccess = true))
 	bool bPlayOnlyWhenVisible = false;
diff --git a/Engine/Plugins/Media/MediaPlate/Source/MediaPlateEditor/Private/MediaPlateEditorModule.cpp b/Engine/Plugins/Media/MediaPlate/Source/MediaPlateEditor/Private/MediaPlateEditorModule.cpp
--- a/Engine/Plugins/Media/MediaPlate/Source/MediaPlateEditor/Private/MediaPlateEditorModule.cpp
+++ b/Engine/Plugins/Media/MediaPlate/Source/MediaPlateEditor/Private/MediaPlateEditorModule.cpp
@@ -99,17 +99,12 @@ void FMediaPlateEditorModule::Tick(float DeltaTime)
 				SoundComponent->UpdatePlayer();
 			}
 
-			// Get the player.
-			TObjectPtr<UMediaPlayer> MediaPlayer = MediaPlate->GetMediaPlayer();
-			if (MediaPlayer != nullptr)
+			bIsMediaPlateToBeRemoved = false;
+			// Are we playing something?
+			if (MediaPlate->IsMediaPlatePlaying())
 			{
-				bIsMediaPlateToBeRemoved = false;
-				// Are we playing something?
-				if ((MediaPlayer->IsPlaying()) || (MediaPlate->GetWantsToPlayWhenVisible()))
-				{
-					bIsMediaPlatePlaying = true;
-					break;
-				}
+				bIsMediaPlatePlaying = true;
+				break;
 			}
 		}
 
diff --git a/Engine/Source/Runtime/MediaAssets/Public/MediaPlayerProxyInterface.h b/Engine/Source/Runtime/MediaAssets/Public/MediaPlayerProxyInterface.h
--- a/Engine/Source/Runtime/MediaAssets/Public/MediaPlayerProxyInterface.h
+++ b/Engine/Source/Runtime/MediaAssets/Public/MediaPlayerProxyInterface.h
@@ -22,6 +22,21 @@ class MEDIAASSETS_API IMediaPlayerProxyInterface
 	GENERATED_BODY()
 
 public:
+	/**
+	 * Get the desired playback rate.
+	 * Note that this is not necessarily the actual rate of the player,
+	 * merely the desired rate the user wants.
+	 */
+	virtual float GetProxyRate() const = 0;
+
+	/**
+	 * Changes the desired playback rate.
+	 *
+	 * @param Rate		The playback rate to set.
+	 * @return			True on success, false otherwise.
+	 */
+	virtual bool SetProxyRate(float Rate) = 0;
+
 	/**
 	 * Call this to see if you can control the media player, or if the owning object is using it.
 	 * 
