MediaCompositing: Media track will use the media cache settings from the proxy if available.
#rnx

diff --git a/Engine/Plugins/Media/MediaCompositing/Source/MediaCompositing/Private/MovieScene/MovieSceneMediaTemplate.cpp b/Engine/Plugins/Media/MediaCompositing/Source/MediaCompositing/Private/MovieScene/MovieSceneMediaTemplate.cpp
--- a/Engine/Plugins/Media/MediaCompositing/Source/MediaCompositing/Private/MovieScene/MovieSceneMediaTemplate.cpp
+++ b/Engine/Plugins/Media/MediaCompositing/Source/MediaCompositing/Private/MovieScene/MovieSceneMediaTemplate.cpp
@@ -40,15 +40,28 @@ struct FMediaSectionPreRollExecutionToken
 
 		FMovieSceneMediaData& SectionData = PersistentData.GetSectionData<FMovieSceneMediaData>();
 		UMediaPlayer* MediaPlayer = SectionData.GetMediaPlayer();
+		UObject* PlayerProxy = SectionData.GetPlayerProxy();
 
 		if (MediaPlayer == nullptr || MediaSource == nullptr)
 		{
 			return;
 		}
 
+		// Do we have a player proxy?
+		IMediaPlayerProxyInterface* PlayerProxyInterface = nullptr;
+		if (PlayerProxy != nullptr)
+		{
+			PlayerProxyInterface = Cast<IMediaPlayerProxyInterface>
+				(PlayerProxy);
+		}
+
 		// open the media source if necessary
 		if (MediaPlayer->GetUrl().IsEmpty())
 		{
+			if (PlayerProxyInterface != nullptr)
+			{
+				MediaSource->SetCacheSettings(PlayerProxyInterface->GetCacheSettings());
+			}
 			SectionData.SeekOnOpen(StartTime);
 			MediaPlayer->OpenSource(MediaSource);
 		}
@@ -101,6 +114,10 @@ struct FMediaSectionExecutionToken
 		// open the media source if necessary
 		if (MediaPlayer->GetUrl().IsEmpty())
 		{
+			if (PlayerProxyInterface != nullptr)
+			{
+				MediaSource->SetCacheSettings(PlayerProxyInterface->GetCacheSettings());
+			}
 			SectionData.SeekOnOpen(CurrentTime);
 			// Setup an initial blocking range - MediaFramework will block (even through the opening process) in its next tick...
 			MediaPlayer->SetBlockOnTimeRange(TRange<FTimespan>(CurrentTime, CurrentTime + FrameDuration));
diff --git a/Engine/Plugins/Media/MediaPlate/Source/MediaPlate/Private/MediaPlateComponent.cpp b/Engine/Plugins/Media/MediaPlate/Source/MediaPlate/Private/MediaPlateComponent.cpp
--- a/Engine/Plugins/Media/MediaPlate/Source/MediaPlate/Private/MediaPlateComponent.cpp
+++ b/Engine/Plugins/Media/MediaPlate/Source/MediaPlate/Private/MediaPlateComponent.cpp
@@ -478,6 +478,11 @@ bool UMediaPlateComponent::IsExternalControlAllowed()
 	return IsVisible();
 }
 
+const FMediaSourceCacheSettings& UMediaPlateComponent::GetCacheSettings() const
+{
+	return CacheSettings;
+}
+
 void UMediaPlateComponent::RestartPlayer()
 {
 	if (MediaPlayer != nullptr)
diff --git a/Engine/Plugins/Media/MediaPlate/Source/MediaPlate/Public/MediaPlateComponent.h b/Engine/Plugins/Media/MediaPlate/Source/MediaPlate/Public/MediaPlateComponent.h
--- a/Engine/Plugins/Media/MediaPlate/Source/MediaPlate/Public/MediaPlateComponent.h
+++ b/Engine/Plugins/Media/MediaPlate/Source/MediaPlate/Public/MediaPlateComponent.h
@@ -226,6 +226,7 @@ public:
 	virtual float GetProxyRate() const override;
 	virtual bool SetProxyRate(float Rate) override;
 	virtual bool IsExternalControlAllowed() override;
+	virtual const FMediaSourceCacheSettings& GetCacheSettings() const override;
 
 private:
 	/**
diff --git a/Engine/Source/Runtime/MediaAssets/Public/MediaPlayerProxyInterface.h b/Engine/Source/Runtime/MediaAssets/Public/MediaPlayerProxyInterface.h
--- a/Engine/Source/Runtime/MediaAssets/Public/MediaPlayerProxyInterface.h
+++ b/Engine/Source/Runtime/MediaAssets/Public/MediaPlayerProxyInterface.h
@@ -6,6 +6,8 @@
 
 #include "MediaPlayerProxyInterface.generated.h"
 
+struct FMediaSourceCacheSettings;
+
 /**
  * The proxy object provides a higher level of insight/control than just the media player.
  * For example, the object owning the player may control the player in some
@@ -43,4 +45,10 @@ public:
 	 * @return				True if you can control the player.
 	 */
 	virtual bool IsExternalControlAllowed() = 0;
+
+	/**
+	 * Gets the cache settings for the player.
+	 */
+	virtual const FMediaSourceCacheSettings& GetCacheSettings() const = 0;
+	
 };
