Media: Prevent using empty ranges on media section evaluation.

With zero-length frames, we restore behavior prior to 21706277 by taking the max of the updated calculation with the previously fixed frame time. This mitigates issues with empty ranges that can occur when "Lock to Display Rate at Runtime" is enabled.

diff --git a/Engine/Plugins/Media/MediaCompositing/Source/MediaCompositing/Private/MovieScene/MovieSceneMediaTemplate.cpp b/Engine/Plugins/Media/MediaCompositing/Source/MediaCompositing/Private/MovieScene/MovieSceneMediaTemplate.cpp
--- a/Engine/Plugins/Media/MediaCompositing/Source/MediaCompositing/Private/MovieScene/MovieSceneMediaTemplate.cpp
+++ b/Engine/Plugins/Media/MediaCompositing/Source/MediaCompositing/Private/MovieScene/MovieSceneMediaTemplate.cpp
@@ -307,11 +307,12 @@ void FMovieSceneMediaSectionTemplate::Evaluate(const FMovieSceneEvaluationOperan
 		const double FrameTimeInSeconds = FrameRate.AsSeconds(FrameTime);
 		const int64 FrameTicks = FrameTimeInSeconds * ETimespan::TicksPerSecond;
 
-		const double FrameDurationInSeconds = (Context.GetRange().Size<FFrameTime>()) / Context.GetFrameRate();
+		// With zero-length frames (which can occur occasionally), we use the fixed frame time, matching previous behavior.
+		const double FrameDurationInSeconds = FMath::Max(FrameRate.AsSeconds(FFrameTime(1)), (Context.GetRange().Size<FFrameTime>()) / Context.GetFrameRate());
 		const int64 FrameDurationTicks = FrameDurationInSeconds * ETimespan::TicksPerSecond;
 
 		#if MOVIESCENEMEDIATEMPLATE_TRACE_EVALUATION
-			GLog->Logf(ELogVerbosity::Log, TEXT("Evaluating frame %i+%f, FrameRate %i/%i, FrameTicks %d, FrameDurationTicks"),
+			GLog->Logf(ELogVerbosity::Log, TEXT("Evaluating frame %i+%f, FrameRate %i/%i, FrameTicks %d, FrameDurationTicks %d"),
 				Context.GetTime().GetFrame().Value,
 				Context.GetTime().GetSubFrame(),
 				FrameRate.Numerator,
