MediaCompositing: Media cache settings are hidden if a player proxy is present as the settings come from the proxy.
#rnx

diff --git a/Engine/Plugins/Media/MediaCompositing/Source/MediaCompositing/Private/MovieScene/MovieSceneMediaSection.cpp b/Engine/Plugins/Media/MediaCompositing/Source/MediaCompositing/Private/MovieScene/MovieSceneMediaSection.cpp
--- a/Engine/Plugins/Media/MediaCompositing/Source/MediaCompositing/Private/MovieScene/MovieSceneMediaSection.cpp
+++ b/Engine/Plugins/Media/MediaCompositing/Source/MediaCompositing/Private/MovieScene/MovieSceneMediaSection.cpp
@@ -22,6 +22,7 @@ namespace
 UMovieSceneMediaSection::UMovieSceneMediaSection(const FObjectInitializer& ObjectInitializer)
 	: Super(ObjectInitializer)
 	, bLooping(true)
+	, bHasMediaPlayerProxy(false)
 {
 #if WITH_EDITORONLY_DATA
 	ThumbnailReferenceOffset = 0.f;
diff --git a/Engine/Plugins/Media/MediaCompositing/Source/MediaCompositing/Public/MovieSceneMediaSection.h b/Engine/Plugins/Media/MediaCompositing/Source/MediaCompositing/Public/MovieSceneMediaSection.h
--- a/Engine/Plugins/Media/MediaCompositing/Source/MediaCompositing/Public/MovieSceneMediaSection.h
+++ b/Engine/Plugins/Media/MediaCompositing/Source/MediaCompositing/Public/MovieSceneMediaSection.h
@@ -55,10 +55,14 @@ public:
 	UPROPERTY(EditAnywhere, BlueprintReadWrite, Category="Media", meta = (EditCondition = "bUseExternalMediaPlayer"), AdvancedDisplay)
 	TObjectPtr<UMediaPlayer> ExternalMediaPlayer;
 
-	/** Override the default cache settings. */
-	UPROPERTY(EditAnywhere, BlueprintReadWrite, Category = "Media", AdvancedDisplay)
+	/** Override the default cache settings. Not used if we have a player proxy as the settings come from the proxy instead. */
+	UPROPERTY(EditAnywhere, BlueprintReadWrite, Category = "Media", meta = (EditCondition = "!bHasMediaPlayerProxy", HideEditConditionToggle = true, EditConditionHides), AdvancedDisplay)
 	FMediaSourceCacheSettings CacheSettings;
 
+	/** True if the object bound to this track has a media player proxy. */
+	UPROPERTY(Transient)
+	bool bHasMediaPlayerProxy;
+
 public:
 
 	/**
diff --git a/Engine/Plugins/Media/MediaCompositing/Source/MediaCompositingEditor/Private/Sequencer/MediaTrackEditor.cpp b/Engine/Plugins/Media/MediaCompositing/Source/MediaCompositingEditor/Private/Sequencer/MediaTrackEditor.cpp
--- a/Engine/Plugins/Media/MediaCompositing/Source/MediaCompositingEditor/Private/Sequencer/MediaTrackEditor.cpp
+++ b/Engine/Plugins/Media/MediaCompositing/Source/MediaCompositingEditor/Private/Sequencer/MediaTrackEditor.cpp
@@ -7,6 +7,7 @@
 #include "Framework/Application/SlateApplication.h"
 #include "Framework/MultiBox/MultiBoxBuilder.h"
 #include "IContentBrowserSingleton.h"
+#include "IMediaAssetsModule.h"
 #include "ISequencer.h"
 #include "ISequencerObjectChangeListener.h"
 #include "MediaSource.h"
@@ -45,6 +46,7 @@ FMediaTrackEditor::FMediaTrackEditor(TSharedRef<ISequencer> InSequencer)
 	: FMovieSceneTrackEditor(InSequencer)
 {
 	ThumbnailPool = MakeShared<FTrackEditorThumbnailPool>(InSequencer);
+	InSequencer->OnMovieSceneDataChanged().AddRaw(this, &FMediaTrackEditor::OnSequencerDataChanged);
 }
 
 
@@ -203,6 +205,14 @@ const FSlateBrush* FMediaTrackEditor::GetIconBrush() const
 	return FAppStyle::GetBrush("Sequencer.Tracks.Media");
 }
 
+void FMediaTrackEditor::OnRelease()
+{
+	const TSharedPtr<ISequencer> SequencerPtr = GetSequencer();
+	if (SequencerPtr.IsValid())
+	{
+		SequencerPtr->OnMovieSceneDataChanged().RemoveAll(this);
+	}
+}
 
 /* FMediaTrackEditor implementation
  *****************************************************************************/
@@ -332,6 +342,62 @@ void FMediaTrackEditor::HandleAddMediaTrackMenuEntryExecute()
 	}
 }
 
+void FMediaTrackEditor::OnSequencerDataChanged(EMovieSceneDataChangeType DataChangeType)
+{
+	// Get the scene.
+	const TSharedPtr<ISequencer> SequencerPtr = GetSequencer();
+	
+	if (SequencerPtr.IsValid())
+	{
+		UMovieSceneSequence* MovieSequence = SequencerPtr->GetFocusedMovieSceneSequence();
+		if (MovieSequence != nullptr)
+		{
+			UMovieScene* MovieScene = MovieSequence->GetMovieScene();
 
+			// Loop through all bindings.
+			const TArray<FMovieSceneBinding>& Bindings = MovieScene->GetBindings();
+			for (const FMovieSceneBinding& Binding : Bindings)
+			{
+				UMovieSceneMediaTrack* Track = Cast<UMovieSceneMediaTrack>(MovieScene->FindTrack(UMovieSceneMediaTrack::StaticClass(), Binding.GetObjectGuid(), NAME_None));
+				if (Track != nullptr)
+				{
+					UpdateMediaTrackBinding(Track, Binding);
+				}
+			}
+		}
+	}
+}
+
+void FMediaTrackEditor::UpdateMediaTrackBinding(UMovieSceneMediaTrack* MediaTrack, const FMovieSceneBinding& Binding)
+{
+	// Get object from the binding.
+	FGuid ObjectHandle = Binding.GetObjectGuid();
+	const TSharedPtr<ISequencer> SequencerPtr = GetSequencer();
+	UObject* BoundObject = SequencerPtr.IsValid() ?
+		SequencerPtr->FindSpawnedObjectOrTemplate(ObjectHandle) : nullptr;
+
+	// Get the player proxy if available.
+	UObject* PlayerProxy = nullptr;
+	if (BoundObject != nullptr)
+	{
+		IMediaAssetsModule* MediaAssetsModule = FModuleManager::LoadModulePtr<IMediaAssetsModule>("MediaAssets");
+		if (MediaAssetsModule != nullptr)
+		{
+			MediaAssetsModule->GetPlayerFromObject(BoundObject, PlayerProxy);
+		}
+	}
+	bool bHasPlayerProxy = PlayerProxy != nullptr;
+
+	// Update all sections.
+	const TArray<UMovieSceneSection*> Sections = MediaTrack->GetAllSections();
+	for (UMovieSceneSection* Section : Sections)
+	{
+		UMovieSceneMediaSection* MediaSection = Cast<UMovieSceneMediaSection>(Section);
+		if (MediaSection != nullptr)
+		{
+			MediaSection->bHasMediaPlayerProxy = bHasPlayerProxy;
+		}
+	}
+}
 
 #undef LOCTEXT_NAMESPACE
diff --git a/Engine/Plugins/Media/MediaCompositing/Source/MediaCompositingEditor/Public/Sequencer/MediaTrackEditor.h b/Engine/Plugins/Media/MediaCompositing/Source/MediaCompositingEditor/Public/Sequencer/MediaTrackEditor.h
--- a/Engine/Plugins/Media/MediaCompositing/Source/MediaCompositingEditor/Public/Sequencer/MediaTrackEditor.h
+++ b/Engine/Plugins/Media/MediaCompositing/Source/MediaCompositingEditor/Public/Sequencer/MediaTrackEditor.h
@@ -71,6 +71,7 @@ public:
 	virtual bool SupportsType(TSubclassOf<UMovieSceneTrack> TrackClass) const override;
 	virtual void Tick(float DeltaTime) override;
 	virtual const FSlateBrush* GetIconBrush() const override;
+	virtual void OnRelease() override;
 
 protected:
 
@@ -89,6 +90,16 @@ private:
 
 	/** Callback for executing the "Add Media Track" menu entry. */
 	void HandleAddMediaTrackMenuEntryExecute();
+	/** Callback for when some sequencer data like bindings change. */
+	void OnSequencerDataChanged(EMovieSceneDataChangeType DataChangeType);
+
+	/**
+	 * Updates a media track with the current binding information.
+	 * 
+	 * @param MediaTrack			Track to update.
+	 * @param Binding				Binding to get object from.
+	 */
+	void UpdateMediaTrackBinding(UMovieSceneMediaTrack* MediaTrack, const FMovieSceneBinding& Binding);
 
 private:
 
diff --git a/Engine/Plugins/Media/MediaPlate/Source/MediaPlateEditor/Private/Sequencer/MediaPlateTrackEditor.cpp b/Engine/Plugins/Media/MediaPlate/Source/MediaPlateEditor/Private/Sequencer/MediaPlateTrackEditor.cpp
--- a/Engine/Plugins/Media/MediaPlate/Source/MediaPlateEditor/Private/Sequencer/MediaPlateTrackEditor.cpp
+++ b/Engine/Plugins/Media/MediaPlate/Source/MediaPlateEditor/Private/Sequencer/MediaPlateTrackEditor.cpp
@@ -194,11 +194,11 @@ void FMediaPlateTrackEditor::AddTrackForComponent(UMediaPlateComponent* Componen
 						// Start process to get the duration.
 						StartGetDuration(MediaSource, Section);
 
-						// Copy cache settings from media plate.
+						// Tell the section it has a player proxy.
 						UMovieSceneMediaSection* MediaSection = Cast<UMovieSceneMediaSection>(Section);
 						if (MediaSection != nullptr)
 						{
-							MediaSection->CacheSettings = Component->CacheSettings;
+							MediaSection->bHasMediaPlayerProxy = true;
 						}
 					}
 				}
