MediaCompositing: Fixed OnMovieSceneDataChanged not getting unbound when OnRelease is not called.
#rnx

diff --git a/Engine/Plugins/Media/MediaCompositing/Source/MediaCompositingEditor/Private/Sequencer/MediaTrackEditor.cpp b/Engine/Plugins/Media/MediaCompositing/Source/MediaCompositingEditor/Private/Sequencer/MediaTrackEditor.cpp
--- a/Engine/Plugins/Media/MediaCompositing/Source/MediaCompositingEditor/Private/Sequencer/MediaTrackEditor.cpp
+++ b/Engine/Plugins/Media/MediaCompositing/Source/MediaCompositingEditor/Private/Sequencer/MediaTrackEditor.cpp
@@ -52,6 +52,11 @@ FMediaTrackEditor::FMediaTrackEditor(TSharedRef<ISequencer> InSequencer)
 
 FMediaTrackEditor::~FMediaTrackEditor()
 {
+	const TSharedPtr<ISequencer> SequencerPtr = GetSequencer();
+	if (SequencerPtr.IsValid())
+	{
+		SequencerPtr->OnMovieSceneDataChanged().RemoveAll(this);
+	}
 }
 
 
