Media: Fixed looping indicators when the media section is offset.

diff --git a/Engine/Plugins/Media/MediaCompositing/Source/MediaCompositingEditor/Private/Sequencer/MediaThumbnailSection.cpp b/Engine/Plugins/Media/MediaCompositing/Source/MediaCompositingEditor/Private/Sequencer/MediaThumbnailSection.cpp
--- a/Engine/Plugins/Media/MediaCompositing/Source/MediaCompositingEditor/Private/Sequencer/MediaThumbnailSection.cpp
+++ b/Engine/Plugins/Media/MediaCompositing/Source/MediaCompositingEditor/Private/Sequencer/MediaThumbnailSection.cpp
@@ -291,7 +291,8 @@ void FMediaThumbnailSection::DrawLoopIndicators(FSequencerSectionPainter& InPain
 	FFrameRate TickResolution = Section->GetTypedOuter<UMovieScene>()->GetTickResolution();
 	double SectionDuration = FFrameTime(UE::MovieScene::DiscreteSize(Section->GetRange())) / TickResolution;
 	const float MediaSizeX = MediaDuration.GetTotalSeconds() * SectionSize.X / SectionDuration;
-	float DrawOffset = MediaSizeX - TimeToPixelConverter.SecondsToPixel(TickResolution.AsSeconds(MediaSection->StartFrameOffset));
+	const FFrameNumber SectionOffset = MediaSection->GetRange().HasLowerBound() ? MediaSection->GetRange().GetLowerBoundValue() : 0;
+	float DrawOffset = MediaSizeX - TimeToPixelConverter.SecondsToPixel(TickResolution.AsSeconds(SectionOffset + MediaSection->StartFrameOffset));
 
 	while (DrawOffset < SectionSize.X)
 	{
