MediaCompositing: Will no longer close the player if a media plate is looping as it will probably be opened again soon.

diff --git a/Engine/Plugins/Media/MediaCompositing/Source/MediaCompositing/Private/MovieScene/MovieSceneMediaData.cpp b/Engine/Plugins/Media/MediaCompositing/Source/MediaCompositing/Private/MovieScene/MovieSceneMediaData.cpp
--- a/Engine/Plugins/Media/MediaCompositing/Source/MediaCompositing/Private/MovieScene/MovieSceneMediaData.cpp
+++ b/Engine/Plugins/Media/MediaCompositing/Source/MediaCompositing/Private/MovieScene/MovieSceneMediaData.cpp
@@ -24,7 +24,12 @@ FMovieSceneMediaData::~FMovieSceneMediaData()
 	if (MediaPlayer != nullptr)
 	{
 		MediaPlayer->OnMediaEvent().RemoveAll(this);
-		MediaPlayer->Close();
+		// If we have a proxy and are looping, then don't bother closing as
+		// we are just going to open again soon anyway.
+		if ((PlayerProxy == nullptr) || (MediaPlayer->IsLooping() == false))
+		{
+			MediaPlayer->Close();
+		}
 		MediaPlayer->RemoveFromRoot();
 	}
 }
