MediaCompositing: Media plate will now be linked to the section instead of just passing the media source to resolve issues when the media source is not an asset.

diff --git a/Engine/Plugins/Media/MediaCompositing/Source/MediaCompositing/Private/MovieScene/MovieSceneMediaSection.cpp b/Engine/Plugins/Media/MediaCompositing/Source/MediaCompositing/Private/MovieScene/MovieSceneMediaSection.cpp
--- a/Engine/Plugins/Media/MediaCompositing/Source/MediaCompositing/Private/MovieScene/MovieSceneMediaSection.cpp
+++ b/Engine/Plugins/Media/MediaCompositing/Source/MediaCompositing/Private/MovieScene/MovieSceneMediaSection.cpp
@@ -1,6 +1,8 @@
 // Copyright Epic Games, Inc. All Rights Reserved.
 
 #include "MovieSceneMediaSection.h"
+
+#include "MediaPlayerProxyInterface.h"
 #include "MovieScene.h"
 #include "Misc/FrameRate.h"
 
@@ -21,6 +23,7 @@ namespace
 
 UMovieSceneMediaSection::UMovieSceneMediaSection(const FObjectInitializer& ObjectInitializer)
 	: Super(ObjectInitializer)
+	, MediaSourceProxyIndex(0)
 	, bLooping(true)
 	, bHasMediaPlayerProxy(false)
 {
@@ -88,5 +91,28 @@ void UMovieSceneMediaSection::MigrateFrameTimes(FFrameRate SourceRate, FFrameRat
 	}
 }
 
+UMediaSource* UMovieSceneMediaSection::GetMediaSource() const
+{
+	UMediaSource* LocalMediaSource = nullptr;
+
+	// Get media source from the proxy if we have one.
+	UObject* MediaSourceProxyObject = MediaSourceProxy.Get();
+	IMediaPlayerProxyInterface* PlayerProxyInterface = nullptr;
+	if (MediaSourceProxyObject)
+	{
+		PlayerProxyInterface = Cast<IMediaPlayerProxyInterface>(MediaSourceProxyObject);
+	}
+	if (PlayerProxyInterface != nullptr)
+	{
+		LocalMediaSource = PlayerProxyInterface->GetMediaSourceFromIndex(MediaSourceProxyIndex);
+	}
+	else
+	{
+		LocalMediaSource = MediaSource.Get();
+	}
+
+	return LocalMediaSource;
+}
+
 #undef LOCTEXT_NAMESPACE
 
diff --git a/Engine/Plugins/Media/MediaCompositing/Source/MediaCompositing/Private/MovieScene/MovieSceneMediaTemplate.cpp b/Engine/Plugins/Media/MediaCompositing/Source/MediaCompositing/Private/MovieScene/MovieSceneMediaTemplate.cpp
--- a/Engine/Plugins/Media/MediaCompositing/Source/MediaCompositing/Private/MovieScene/MovieSceneMediaTemplate.cpp
+++ b/Engine/Plugins/Media/MediaCompositing/Source/MediaCompositing/Private/MovieScene/MovieSceneMediaTemplate.cpp
@@ -254,6 +254,8 @@ private:
 FMovieSceneMediaSectionTemplate::FMovieSceneMediaSectionTemplate(const UMovieSceneMediaSection& InSection, const UMovieSceneMediaTrack& InTrack)
 {
 	Params.MediaSource = InSection.GetMediaSource();
+	Params.MediaSourceProxy = InSection.MediaSourceProxy.Get();
+	Params.MediaSourceProxyIndex = InSection.MediaSourceProxyIndex;
 	Params.MediaSoundComponent = InSection.MediaSoundComponent;
 	Params.bLooping = InSection.bLooping;
 	Params.StartFrameOffset = InSection.StartFrameOffset;
@@ -282,7 +284,19 @@ FMovieSceneMediaSectionTemplate::FMovieSceneMediaSectionTemplate(const UMovieSce
 
 void FMovieSceneMediaSectionTemplate::Evaluate(const FMovieSceneEvaluationOperand& Operand, const FMovieSceneContext& Context, const FPersistentEvaluationData& PersistentData, FMovieSceneExecutionTokens& ExecutionTokens) const
 {
-	if ((Params.MediaSource == nullptr) || Context.IsPostRoll())
+	UMediaSource* MediaSource = Params.MediaSource;
+	UObject* MediaSourceProxy = Params.MediaSourceProxy.Get();
+	if (MediaSourceProxy != nullptr)
+	{
+		IMediaPlayerProxyInterface* PlayerProxyInterface =
+			Cast<IMediaPlayerProxyInterface>(MediaSourceProxy);
+		if (PlayerProxyInterface != nullptr)
+		{
+			MediaSource = PlayerProxyInterface->GetMediaSourceFromIndex(Params.MediaSourceProxyIndex);
+		}
+	}
+
+	if ((MediaSource == nullptr) || Context.IsPostRoll())
 	{
 		return;
 	}
@@ -297,7 +311,7 @@ void FMovieSceneMediaSectionTemplate::Evaluate(const FMovieSceneEvaluationOperan
 		const double StartFrameInSeconds = FrameRate.AsSeconds(StartFrame);
 		const int64 StartTicks = StartFrameInSeconds * ETimespan::TicksPerSecond;
 
-		ExecutionTokens.Add(FMediaSectionPreRollExecutionToken(Params.MediaSource, FTimespan(StartTicks)));
+		ExecutionTokens.Add(FMediaSectionPreRollExecutionToken(MediaSource, FTimespan(StartTicks)));
 	}
 	else if (!Context.IsPostRoll() && (Context.GetTime().FrameNumber < Params.SectionEndFrame))
 	{
@@ -322,7 +336,7 @@ void FMovieSceneMediaSectionTemplate::Evaluate(const FMovieSceneEvaluationOperan
 			);
 		#endif
 
-		ExecutionTokens.Add(FMediaSectionExecutionToken(Params.MediaSource, FTimespan(FrameTicks), FTimespan(FrameDurationTicks)));
+		ExecutionTokens.Add(FMediaSectionExecutionToken(MediaSource, FTimespan(FrameTicks), FTimespan(FrameDurationTicks)));
 	}
 }
 
diff --git a/Engine/Plugins/Media/MediaCompositing/Source/MediaCompositing/Private/MovieScene/MovieSceneMediaTemplate.h b/Engine/Plugins/Media/MediaCompositing/Source/MediaCompositing/Private/MovieScene/MovieSceneMediaTemplate.h
--- a/Engine/Plugins/Media/MediaCompositing/Source/MediaCompositing/Private/MovieScene/MovieSceneMediaTemplate.h
+++ b/Engine/Plugins/Media/MediaCompositing/Source/MediaCompositing/Private/MovieScene/MovieSceneMediaTemplate.h
@@ -26,6 +26,12 @@ struct FMovieSceneMediaSectionParams
 	UPROPERTY()
 	TObjectPtr<UMediaSource> MediaSource;
 
+	UPROPERTY()
+	TObjectPtr<UObject> MediaSourceProxy;
+
+	UPROPERTY()
+	int32 MediaSourceProxyIndex;
+
 	UPROPERTY()
 	TObjectPtr<UMediaTexture> MediaTexture;
 
@@ -47,6 +53,8 @@ struct FMovieSceneMediaSectionParams
 	FMovieSceneMediaSectionParams()
 		: MediaSoundComponent(nullptr)
 		, MediaSource(nullptr)
+		, MediaSourceProxy(nullptr)
+		, MediaSourceProxyIndex(0)
 		, MediaTexture(nullptr)
 		, MediaPlayer(nullptr)
 		, bLooping(false)
diff --git a/Engine/Plugins/Media/MediaCompositing/Source/MediaCompositing/Private/MovieScene/MovieSceneMediaTrack.cpp b/Engine/Plugins/Media/MediaCompositing/Source/MediaCompositing/Private/MovieScene/MovieSceneMediaTrack.cpp
--- a/Engine/Plugins/Media/MediaCompositing/Source/MediaCompositing/Private/MovieScene/MovieSceneMediaTrack.cpp
+++ b/Engine/Plugins/Media/MediaCompositing/Source/MediaCompositing/Private/MovieScene/MovieSceneMediaTrack.cpp
@@ -38,6 +38,22 @@ UMovieSceneMediaTrack::UMovieSceneMediaTrack(const FObjectInitializer& ObjectIni
  *****************************************************************************/
 
 UMovieSceneSection* UMovieSceneMediaTrack::AddNewMediaSourceOnRow(UMediaSource& MediaSource, FFrameNumber Time, int32 RowIndex)
+{
+	UMovieSceneMediaSection* NewSection = AddNewSectionOnRow(Time, RowIndex);
+	NewSection->SetMediaSource(&MediaSource);
+	return NewSection;
+}
+
+
+UMovieSceneSection* UMovieSceneMediaTrack::AddNewMediaSourceProxyOnRow(UObject& MediaSourceProxy, int32 MediaSourceProxyIndex, FFrameNumber Time, int32 RowIndex)
+{
+	UMovieSceneMediaSection* NewSection = AddNewSectionOnRow(Time, RowIndex);
+	NewSection->SetMediaSourceProxy(&MediaSourceProxy, MediaSourceProxyIndex);
+	return NewSection;
+}
+
+
+UMovieSceneMediaSection* UMovieSceneMediaTrack::AddNewSectionOnRow(FFrameNumber Time, int32 RowIndex)
 {
 	const float DefaultMediaSectionDuration = 1.0f;
 	FFrameRate TickResolution = GetTypedOuter<UMovieScene>()->GetTickResolution();
@@ -47,7 +63,6 @@ UMovieSceneSection* UMovieSceneMediaTrack::AddNewMediaSourceOnRow(UMediaSource&
 	UMovieSceneMediaSection* NewSection = NewObject<UMovieSceneMediaSection>(this, NAME_None, RF_Transactional);
 
 	NewSection->InitialPlacementOnRow(MediaSections, Time, DurationToUse.FrameNumber.Value, RowIndex);
-	NewSection->SetMediaSource(&MediaSource);
 
 	MediaSections.Add(NewSection);
 
diff --git a/Engine/Plugins/Media/MediaCompositing/Source/MediaCompositing/Public/MovieSceneMediaSection.h b/Engine/Plugins/Media/MediaCompositing/Source/MediaCompositing/Public/MovieSceneMediaSection.h
--- a/Engine/Plugins/Media/MediaCompositing/Source/MediaCompositing/Public/MovieSceneMediaSection.h
+++ b/Engine/Plugins/Media/MediaCompositing/Source/MediaCompositing/Public/MovieSceneMediaSection.h
@@ -27,9 +27,15 @@ public:
 
 
 public:
-	/** The source to play with this video track. */
+	/** The source to play with this video track if MediaSourceProxy is not available. */
 	UPROPERTY(EditAnywhere, BlueprintReadWrite, Category="Media")
 	TObjectPtr<UMediaSource> MediaSource;
+	/** The object to get the source to play from if you don't want to directly specify a media source. */
+	UPROPERTY(EditAnywhere, BlueprintReadWrite, Category = "Media")
+	TSoftObjectPtr<UObject> MediaSourceProxy;
+	/** The index to pass to MediaSourceProxy to get the media source. */
+	UPROPERTY(EditAnywhere, BlueprintReadWrite, Category = "Media")
+	int32 MediaSourceProxyIndex;
 
 	/** Should the media player be set to loop? This can be helpful for media formats that can use this information (such as exr sequences) to pre-cache the starting data when nearing the end of playback. Does not cause the media to continue playing after the end of the section is reached. */
 	UPROPERTY(EditAnywhere, BlueprintReadWrite, Category = "Media")
@@ -90,10 +96,7 @@ public:
 	 * @return The media source object.
 	 * @see SetMediaSource
 	 */
-	UMediaSource* GetMediaSource() const
-	{
-		return MediaSource;
-	}
+	UMediaSource* GetMediaSource() const;
 
 	/**
 	 * Set this section's video source.
@@ -106,6 +109,18 @@ public:
 		MediaSource = InMediaSource;
 	}
 
+	/**
+	 * Set this section's video source proxy.
+	 *
+	 * @param InMediaSourceProxy	This should implement IMediaPlayerProxyInterface.
+	 * @param InIndex				Index to pass to IMediaPlayerProxyInterface.
+	 */
+	void SetMediaSourceProxy(UObject* InMediaSourceProxy, int32 InIndex)
+	{
+		MediaSourceProxy = InMediaSourceProxy;
+		MediaSourceProxyIndex = InIndex;
+	}
+
 #if WITH_EDITORONLY_DATA
 
 	/** @return The thumbnail reference frame offset from the start of this section */
diff --git a/Engine/Plugins/Media/MediaCompositing/Source/MediaCompositing/Public/MovieSceneMediaTrack.h b/Engine/Plugins/Media/MediaCompositing/Source/MediaCompositing/Public/MovieSceneMediaTrack.h
--- a/Engine/Plugins/Media/MediaCompositing/Source/MediaCompositing/Public/MovieSceneMediaTrack.h
+++ b/Engine/Plugins/Media/MediaCompositing/Source/MediaCompositing/Public/MovieSceneMediaTrack.h
@@ -10,6 +10,7 @@
 #include "MovieSceneMediaTrack.generated.h"
 
 class UMediaSource;
+class UMovieSceneMediaSection;
 
 
 /**
@@ -35,9 +36,16 @@ public:
 
 	/** Adds a new media source to the track. */
 	virtual UMovieSceneSection* AddNewMediaSourceOnRow(UMediaSource& MediaSource, FFrameNumber Time, int32 RowIndex);
+	/** Adds a new media source to the track. */
+	virtual UMovieSceneSection* AddNewMediaSourceProxyOnRow(UObject& MediaSourceProxy, int32 MediaSourceProxyIndex, FFrameNumber Time, int32 RowIndex);
 
 	/** Adds a new media source on the next available/non-overlapping row. */
 	virtual UMovieSceneSection* AddNewMediaSource(UMediaSource& MediaSource, FFrameNumber Time) { return AddNewMediaSourceOnRow(MediaSource, Time, INDEX_NONE); }
+	/** Adds a new media source on the next available/non-overlapping row. */
+	virtual UMovieSceneSection* AddNewMediaSourceProxy(UObject& MediaSourceProxy, int32 MediaSourceProxyIndex, FFrameNumber Time)
+	{
+		return AddNewMediaSourceProxyOnRow(MediaSourceProxy, MediaSourceProxyIndex, Time, INDEX_NONE); 
+	}
 
 public:
 
@@ -55,6 +63,8 @@ public:
 	virtual bool SupportsMultipleRows() const override { return true; }
 
 private:
+	/** Base function to add a new section. */
+	UMovieSceneMediaSection* AddNewSectionOnRow(FFrameNumber Time, int32 RowIndex);
 
 	/** List of all media sections. */
 	UPROPERTY()
diff --git a/Engine/Plugins/Media/MediaPlate/Source/MediaPlate/Private/MediaPlateComponent.cpp b/Engine/Plugins/Media/MediaPlate/Source/MediaPlate/Private/MediaPlateComponent.cpp
--- a/Engine/Plugins/Media/MediaPlate/Source/MediaPlate/Private/MediaPlateComponent.cpp
+++ b/Engine/Plugins/Media/MediaPlate/Source/MediaPlate/Private/MediaPlateComponent.cpp
@@ -581,6 +581,16 @@ const FMediaSourceCacheSettings& UMediaPlateComponent::GetCacheSettings() const
 	return CacheSettings;
 }
 
+UMediaSource* UMediaPlateComponent::GetMediaSourceFromIndex(int32 Index) const
+{
+	UMediaSource* MediaSource = nullptr;
+	if (MediaPlaylist != nullptr)
+	{
+		MediaSource = MediaPlaylist->Get(Index);
+	}
+	return MediaSource;
+}
+
 void UMediaPlateComponent::RestartPlayer()
 {
 	if (MediaPlayer != nullptr)
diff --git a/Engine/Plugins/Media/MediaPlate/Source/MediaPlate/Public/MediaPlateComponent.h b/Engine/Plugins/Media/MediaPlate/Source/MediaPlate/Public/MediaPlateComponent.h
--- a/Engine/Plugins/Media/MediaPlate/Source/MediaPlate/Public/MediaPlateComponent.h
+++ b/Engine/Plugins/Media/MediaPlate/Source/MediaPlate/Public/MediaPlateComponent.h
@@ -241,6 +241,7 @@ public:
 	virtual bool SetProxyRate(float Rate) override;
 	virtual bool IsExternalControlAllowed() override;
 	virtual const FMediaSourceCacheSettings& GetCacheSettings() const override;
+	virtual UMediaSource* GetMediaSourceFromIndex(int32 Index) const override;
 
 private:
 	/**
diff --git a/Engine/Plugins/Media/MediaPlate/Source/MediaPlateEditor/Private/Sequencer/MediaPlateTrackEditor.cpp b/Engine/Plugins/Media/MediaPlate/Source/MediaPlateEditor/Private/Sequencer/MediaPlateTrackEditor.cpp
--- a/Engine/Plugins/Media/MediaPlate/Source/MediaPlateEditor/Private/Sequencer/MediaPlateTrackEditor.cpp
+++ b/Engine/Plugins/Media/MediaPlate/Source/MediaPlateEditor/Private/Sequencer/MediaPlateTrackEditor.cpp
@@ -188,7 +188,7 @@ void FMediaPlateTrackEditor::AddTrackForComponent(UMediaPlateComponent* Componen
 				UMediaSource* MediaSource = Playlist->Get(Index);
 				if (MediaSource != nullptr)
 				{
-					UMovieSceneSection* Section = MediaTrack->AddNewMediaSource(*MediaSource, FFrameNumber(0));
+					UMovieSceneSection* Section = MediaTrack->AddNewMediaSourceProxy(*Component, Index, FFrameNumber(0));
 					if (Section != nullptr)
 					{
 						// Start process to get the duration.
diff --git a/Engine/Source/Runtime/MediaAssets/Public/MediaPlayerProxyInterface.h b/Engine/Source/Runtime/MediaAssets/Public/MediaPlayerProxyInterface.h
--- a/Engine/Source/Runtime/MediaAssets/Public/MediaPlayerProxyInterface.h
+++ b/Engine/Source/Runtime/MediaAssets/Public/MediaPlayerProxyInterface.h
@@ -6,6 +6,7 @@
 
 #include "MediaPlayerProxyInterface.generated.h"
 
+class UMediaSource; 
 struct FMediaSourceCacheSettings;
 
 /**
@@ -50,5 +51,10 @@ public:
 	 * Gets the cache settings for the player.
 	 */
 	virtual const FMediaSourceCacheSettings& GetCacheSettings() const = 0;
+
+	/**
+	 * Get the media source for a given index.
+	 */
+	virtual UMediaSource* GetMediaSourceFromIndex(int32 Index) const = 0;
 	
 };
