MediaCompositing: Media source proxy is now referenced via a binding to fix issues with spawnables.

diff --git a/Engine/Plugins/Media/MediaCompositing/Source/MediaCompositing/Private/MovieScene/MovieSceneMediaSection.cpp b/Engine/Plugins/Media/MediaCompositing/Source/MediaCompositing/Private/MovieScene/MovieSceneMediaSection.cpp
--- a/Engine/Plugins/Media/MediaCompositing/Source/MediaCompositing/Private/MovieScene/MovieSceneMediaSection.cpp
+++ b/Engine/Plugins/Media/MediaCompositing/Source/MediaCompositing/Private/MovieScene/MovieSceneMediaSection.cpp
@@ -2,6 +2,8 @@
 
 #include "MovieSceneMediaSection.h"
 
+#include "Components/ActorComponent.h"
+#include "GameFramework/Actor.h"
 #include "MediaPlayerProxyInterface.h"
 #include "MovieScene.h"
 #include "Misc/FrameRate.h"
@@ -91,27 +93,67 @@ void UMovieSceneMediaSection::MigrateFrameTimes(FFrameRate SourceRate, FFrameRat
 	}
 }
 
-UMediaSource* UMovieSceneMediaSection::GetMediaSource() const
+void UMovieSceneMediaSection::OnBindingIDsUpdated(const TMap<UE::MovieScene::FFixedObjectBindingID, UE::MovieScene::FFixedObjectBindingID>& OldFixedToNewFixedMap, FMovieSceneSequenceID LocalSequenceID, const FMovieSceneSequenceHierarchy* Hierarchy, IMovieScenePlayer& Player)
 {
-	UMediaSource* LocalMediaSource = nullptr;
+	UE::MovieScene::FFixedObjectBindingID FixedBindingID = 
+		MediaSourceProxyBindingID.ResolveToFixed(LocalSequenceID, Player);
 
-	// Get media source from the proxy if we have one.
-	UObject* MediaSourceProxyObject = MediaSourceProxy.Get();
-	IMediaPlayerProxyInterface* PlayerProxyInterface = nullptr;
-	if (MediaSourceProxyObject)
-	{
-		PlayerProxyInterface = Cast<IMediaPlayerProxyInterface>(MediaSourceProxyObject);
-	}
-	if (PlayerProxyInterface != nullptr)
+	if (OldFixedToNewFixedMap.Contains(FixedBindingID))
 	{
-		LocalMediaSource = PlayerProxyInterface->GetMediaSourceFromIndex(MediaSourceProxyIndex);
+		Modify();
+
+		MediaSourceProxyBindingID =
+			OldFixedToNewFixedMap[FixedBindingID].ConvertToRelative(LocalSequenceID, Hierarchy);
 	}
-	else
+}
+
+void UMovieSceneMediaSection::GetReferencedBindings(TArray<FGuid>& OutBindings)
+{
+	OutBindings.Add(MediaSourceProxyBindingID.GetGuid());
+}
+
+UMediaSource* UMovieSceneMediaSection::GetMediaSource() const
+{
+	return MediaSource.Get();
+}
+
+UMediaSource* UMovieSceneMediaSection::GetMediaSourceOrProxy(IMovieScenePlayer& Player, FMovieSceneSequenceID SequenceID) const
+{
+	return GetMediaSourceOrProxy(Player, SequenceID, MediaSource, MediaSourceProxyBindingID, MediaSourceProxyIndex);
+}
+
+UMediaSource* UMovieSceneMediaSection::GetMediaSourceOrProxy(IMovieScenePlayer& InPlayer, FMovieSceneSequenceID InSequenceID, UMediaSource* InMediaSource, const FMovieSceneObjectBindingID& InMediaSourceProxyBindingID, int32 InMediaSourceProxyIndex)
+{
+	UMediaSource* OutMediaSource = InMediaSource;
+
+	for (TWeakObjectPtr<> WeakObject : InMediaSourceProxyBindingID.ResolveBoundObjects(InSequenceID, InPlayer))
 	{
-		LocalMediaSource = MediaSource.Get();
+		if (UObject* Object = WeakObject.Get())
+		{
+			AActor* Actor = Cast<AActor>(Object);
+			if (Actor != nullptr)
+			{
+				// Loop over all the components and see if any are a proxy.
+				TArray<UActorComponent*> ActorComponents;
+				Actor->GetComponents(ActorComponents);
+				for (UActorComponent* Component : ActorComponents)
+				{
+					IMediaPlayerProxyInterface* Proxy = Cast<IMediaPlayerProxyInterface>(Component);
+					if (Proxy != nullptr)
+					{
+						OutMediaSource = Proxy->GetMediaSourceFromIndex(InMediaSourceProxyIndex);
+						break;
+					}
+				}
+				if (OutMediaSource != nullptr)
+				{
+					break;
+				}
+			}
+		}
 	}
 
-	return LocalMediaSource;
+	return OutMediaSource;
 }
 
 #undef LOCTEXT_NAMESPACE
diff --git a/Engine/Plugins/Media/MediaCompositing/Source/MediaCompositing/Private/MovieScene/MovieSceneMediaTemplate.cpp b/Engine/Plugins/Media/MediaCompositing/Source/MediaCompositing/Private/MovieScene/MovieSceneMediaTemplate.cpp
--- a/Engine/Plugins/Media/MediaCompositing/Source/MediaCompositing/Private/MovieScene/MovieSceneMediaTemplate.cpp
+++ b/Engine/Plugins/Media/MediaCompositing/Source/MediaCompositing/Private/MovieScene/MovieSceneMediaTemplate.cpp
@@ -26,11 +26,37 @@
 /* Local helpers
  *****************************************************************************/
 
-struct FMediaSectionPreRollExecutionToken
+/** Base struct for exectution tokens. */
+struct FMediaSectionBaseExecutionToken
 	: IMovieSceneExecutionToken
 {
-	FMediaSectionPreRollExecutionToken(UMediaSource* InMediaSource, FTimespan InStartTimeSeconds)
-		: MediaSource(InMediaSource)
+	FMediaSectionBaseExecutionToken(UMediaSource* InMediaSource, const FMovieSceneObjectBindingID& InMediaSourceProxy, int32 InMediaSourceProxyIndex)
+		: BaseMediaSource(InMediaSource)
+		, MediaSourceProxy(InMediaSourceProxy)
+		, MediaSourceProxyIndex(InMediaSourceProxyIndex)
+	{
+	}
+
+	/**
+	 * Gets the media source from either the proxy binding or the media source.
+	 */
+	UMediaSource* GetMediaSource(IMovieScenePlayer& Player, FMovieSceneSequenceID SequenceID)
+	{
+		return UMovieSceneMediaSection::GetMediaSourceOrProxy(Player, SequenceID,
+			BaseMediaSource, MediaSourceProxy, MediaSourceProxyIndex);
+	}
+
+private:
+	UMediaSource* BaseMediaSource;
+	FMovieSceneObjectBindingID MediaSourceProxy;
+	int32 MediaSourceProxyIndex = 0;
+};
+
+struct FMediaSectionPreRollExecutionToken
+	: FMediaSectionBaseExecutionToken
+{
+	FMediaSectionPreRollExecutionToken(UMediaSource* InMediaSource, FMovieSceneObjectBindingID InMediaSourceProxy, int32 InMediaSourceProxyIndex, FTimespan InStartTimeSeconds)
+		: FMediaSectionBaseExecutionToken(InMediaSource, InMediaSourceProxy, InMediaSourceProxyIndex)
 		, StartTime(InStartTimeSeconds)
 	{ }
 
@@ -41,6 +67,7 @@ struct FMediaSectionPreRollExecutionToken
 		FMovieSceneMediaData& SectionData = PersistentData.GetSectionData<FMovieSceneMediaData>();
 		UMediaPlayer* MediaPlayer = SectionData.GetMediaPlayer();
 		UObject* PlayerProxy = SectionData.GetPlayerProxy();
+		UMediaSource* MediaSource = GetMediaSource(Player, Operand.SequenceID);
 
 		if (MediaPlayer == nullptr || MediaSource == nullptr)
 		{
@@ -69,18 +96,17 @@ struct FMediaSectionPreRollExecutionToken
 
 private:
 
-	UMediaSource* MediaSource;
 	FTimespan StartTime;
 };
 
 
 struct FMediaSectionExecutionToken
-	: IMovieSceneExecutionToken
+	: FMediaSectionBaseExecutionToken
 {
-	FMediaSectionExecutionToken(UMediaSource* InMediaSource, FTimespan InCurrentTime, FTimespan InFrameDuration)
-		: CurrentTime(InCurrentTime)
+	FMediaSectionExecutionToken(UMediaSource* InMediaSource, FMovieSceneObjectBindingID InMediaSourceProxy, int32 InMediaSourceProxyIndex, FTimespan InCurrentTime, FTimespan InFrameDuration)
+		: FMediaSectionBaseExecutionToken(InMediaSource, InMediaSourceProxy, InMediaSourceProxyIndex)
+		, CurrentTime(InCurrentTime)
 		, FrameDuration(InFrameDuration)
-		, MediaSource(InMediaSource)
 		, PlaybackRate(1.0f)
 	{ }
 
@@ -89,6 +115,7 @@ struct FMediaSectionExecutionToken
 		FMovieSceneMediaData& SectionData = PersistentData.GetSectionData<FMovieSceneMediaData>();
 		UMediaPlayer* MediaPlayer = SectionData.GetMediaPlayer();
 		UObject* PlayerProxy = SectionData.GetPlayerProxy();
+		UMediaSource* MediaSource = GetMediaSource(Player, Operand.SequenceID);
 
 		if (MediaPlayer == nullptr || MediaSource == nullptr)
 		{
@@ -243,7 +270,6 @@ private:
 
 	FTimespan CurrentTime;
 	FTimespan FrameDuration;
-	UMediaSource* MediaSource;
 	float PlaybackRate;
 };
 
@@ -254,7 +280,7 @@ private:
 FMovieSceneMediaSectionTemplate::FMovieSceneMediaSectionTemplate(const UMovieSceneMediaSection& InSection, const UMovieSceneMediaTrack& InTrack)
 {
 	Params.MediaSource = InSection.GetMediaSource();
-	Params.MediaSourceProxy = InSection.MediaSourceProxy.Get();
+	Params.MediaSourceProxy = InSection.GetMediaSourceProxy();
 	Params.MediaSourceProxyIndex = InSection.MediaSourceProxyIndex;
 	Params.MediaSoundComponent = InSection.MediaSoundComponent;
 	Params.bLooping = InSection.bLooping;
@@ -285,18 +311,7 @@ FMovieSceneMediaSectionTemplate::FMovieSceneMediaSectionTemplate(const UMovieSce
 void FMovieSceneMediaSectionTemplate::Evaluate(const FMovieSceneEvaluationOperand& Operand, const FMovieSceneContext& Context, const FPersistentEvaluationData& PersistentData, FMovieSceneExecutionTokens& ExecutionTokens) const
 {
 	UMediaSource* MediaSource = Params.MediaSource;
-	UObject* MediaSourceProxy = Params.MediaSourceProxy.Get();
-	if (MediaSourceProxy != nullptr)
-	{
-		IMediaPlayerProxyInterface* PlayerProxyInterface =
-			Cast<IMediaPlayerProxyInterface>(MediaSourceProxy);
-		if (PlayerProxyInterface != nullptr)
-		{
-			MediaSource = PlayerProxyInterface->GetMediaSourceFromIndex(Params.MediaSourceProxyIndex);
-		}
-	}
-
-	if ((MediaSource == nullptr) || Context.IsPostRoll())
+	if (((MediaSource == nullptr) && (Params.MediaSourceProxy.IsValid() == false)) || Context.IsPostRoll())
 	{
 		return;
 	}
@@ -311,7 +326,7 @@ void FMovieSceneMediaSectionTemplate::Evaluate(const FMovieSceneEvaluationOperan
 		const double StartFrameInSeconds = FrameRate.AsSeconds(StartFrame);
 		const int64 StartTicks = StartFrameInSeconds * ETimespan::TicksPerSecond;
 
-		ExecutionTokens.Add(FMediaSectionPreRollExecutionToken(MediaSource, FTimespan(StartTicks)));
+		ExecutionTokens.Add(FMediaSectionPreRollExecutionToken(MediaSource, Params.MediaSourceProxy, Params.MediaSourceProxyIndex, FTimespan(StartTicks)));
 	}
 	else if (!Context.IsPostRoll() && (Context.GetTime().FrameNumber < Params.SectionEndFrame))
 	{
@@ -336,7 +351,7 @@ void FMovieSceneMediaSectionTemplate::Evaluate(const FMovieSceneEvaluationOperan
 			);
 		#endif
 
-		ExecutionTokens.Add(FMediaSectionExecutionToken(MediaSource, FTimespan(FrameTicks), FTimespan(FrameDurationTicks)));
+		ExecutionTokens.Add(FMediaSectionExecutionToken(MediaSource, Params.MediaSourceProxy, Params.MediaSourceProxyIndex, FTimespan(FrameTicks), FTimespan(FrameDurationTicks)));
 	}
 }
 
diff --git a/Engine/Plugins/Media/MediaCompositing/Source/MediaCompositing/Private/MovieScene/MovieSceneMediaTemplate.h b/Engine/Plugins/Media/MediaCompositing/Source/MediaCompositing/Private/MovieScene/MovieSceneMediaTemplate.h
--- a/Engine/Plugins/Media/MediaCompositing/Source/MediaCompositing/Private/MovieScene/MovieSceneMediaTemplate.h
+++ b/Engine/Plugins/Media/MediaCompositing/Source/MediaCompositing/Private/MovieScene/MovieSceneMediaTemplate.h
@@ -4,6 +4,7 @@
 
 #include "Evaluation/MovieSceneEvalTemplate.h"
 #include "Evaluation/MovieScenePropertyTemplate.h"
+#include "MovieSceneObjectBindingID.h"
 
 #include "MovieSceneMediaTemplate.generated.h"
 
@@ -27,7 +28,7 @@ struct FMovieSceneMediaSectionParams
 	TObjectPtr<UMediaSource> MediaSource;
 
 	UPROPERTY()
-	TObjectPtr<UObject> MediaSourceProxy;
+	FMovieSceneObjectBindingID MediaSourceProxy;
 
 	UPROPERTY()
 	int32 MediaSourceProxyIndex;
@@ -53,7 +54,6 @@ struct FMovieSceneMediaSectionParams
 	FMovieSceneMediaSectionParams()
 		: MediaSoundComponent(nullptr)
 		, MediaSource(nullptr)
-		, MediaSourceProxy(nullptr)
 		, MediaSourceProxyIndex(0)
 		, MediaTexture(nullptr)
 		, MediaPlayer(nullptr)
diff --git a/Engine/Plugins/Media/MediaCompositing/Source/MediaCompositing/Private/MovieScene/MovieSceneMediaTrack.cpp b/Engine/Plugins/Media/MediaCompositing/Source/MediaCompositing/Private/MovieScene/MovieSceneMediaTrack.cpp
--- a/Engine/Plugins/Media/MediaCompositing/Source/MediaCompositing/Private/MovieScene/MovieSceneMediaTrack.cpp
+++ b/Engine/Plugins/Media/MediaCompositing/Source/MediaCompositing/Private/MovieScene/MovieSceneMediaTrack.cpp
@@ -45,10 +45,10 @@ UMovieSceneSection* UMovieSceneMediaTrack::AddNewMediaSourceOnRow(UMediaSource&
 }
 
 
-UMovieSceneSection* UMovieSceneMediaTrack::AddNewMediaSourceProxyOnRow(UObject& MediaSourceProxy, int32 MediaSourceProxyIndex, FFrameNumber Time, int32 RowIndex)
+UMovieSceneSection* UMovieSceneMediaTrack::AddNewMediaSourceProxyOnRow(const FMovieSceneObjectBindingID& ObjectBinding, int32 MediaSourceProxyIndex, FFrameNumber Time, int32 RowIndex)
 {
 	UMovieSceneMediaSection* NewSection = AddNewSectionOnRow(Time, RowIndex);
-	NewSection->SetMediaSourceProxy(&MediaSourceProxy, MediaSourceProxyIndex);
+	NewSection->SetMediaSourceProxy(ObjectBinding, MediaSourceProxyIndex);
 	return NewSection;
 }
 
diff --git a/Engine/Plugins/Media/MediaCompositing/Source/MediaCompositing/Public/MovieSceneMediaSection.h b/Engine/Plugins/Media/MediaCompositing/Source/MediaCompositing/Public/MovieSceneMediaSection.h
--- a/Engine/Plugins/Media/MediaCompositing/Source/MediaCompositing/Public/MovieSceneMediaSection.h
+++ b/Engine/Plugins/Media/MediaCompositing/Source/MediaCompositing/Public/MovieSceneMediaSection.h
@@ -5,14 +5,16 @@
 #include "MovieSceneSection.h"
 
 #include "MediaSource.h"
+#include "MovieSceneObjectBindingID.h"
 
 #include "MovieSceneMediaSection.generated.h"
 
+class IMovieScenePlayer;
 class UMediaPlayer;
 class UMediaSoundComponent;
 class UMediaSource;
 class UMediaTexture;
-
+struct FMovieSceneSequenceID;
 
 /**
  * Implements a movie scene section for media playback.
@@ -30,9 +32,7 @@ public:
 	/** The source to play with this video track if MediaSourceProxy is not available. */
 	UPROPERTY(EditAnywhere, BlueprintReadWrite, Category="Media")
 	TObjectPtr<UMediaSource> MediaSource;
-	/** The object to get the source to play from if you don't want to directly specify a media source. */
-	UPROPERTY(EditAnywhere, BlueprintReadWrite, Category = "Media")
-	TSoftObjectPtr<UObject> MediaSourceProxy;
+	
 	/** The index to pass to MediaSourceProxy to get the media source. */
 	UPROPERTY(EditAnywhere, BlueprintReadWrite, Category = "Media")
 	int32 MediaSourceProxyIndex;
@@ -87,6 +87,9 @@ protected:
 	virtual UMovieSceneSection* SplitSection(FQualifiedFrameTime SplitTime, bool bDeleteKeys) override;
 	virtual TOptional<FFrameTime> GetOffsetTime() const override;
 	virtual void MigrateFrameTimes(FFrameRate SourceRate, FFrameRate DestinationRate) override;
+	virtual void OnBindingIDsUpdated(const TMap<UE::MovieScene::FFixedObjectBindingID, UE::MovieScene::FFixedObjectBindingID>& OldFixedToNewFixedMap, FMovieSceneSequenceID LocalSequenceID, const FMovieSceneSequenceHierarchy* Hierarchy, IMovieScenePlayer& Player) override;
+	virtual void GetReferencedBindings(TArray<FGuid>& OutBindings) override;
+
 
 public:
 
@@ -109,18 +112,44 @@ public:
 		MediaSource = InMediaSource;
 	}
 
+	/**
+	 * Get this section's video source proxy.
+	 */
+	const FMovieSceneObjectBindingID& GetMediaSourceProxy() const
+	{
+		return MediaSourceProxyBindingID;
+	}
+
 	/**
 	 * Set this section's video source proxy.
 	 *
-	 * @param InMediaSourceProxy	This should implement IMediaPlayerProxyInterface.
+	 * @param InObjectBinding		The object this binds to should implement IMediaPlayerProxyInterface.
 	 * @param InIndex				Index to pass to IMediaPlayerProxyInterface.
 	 */
-	void SetMediaSourceProxy(UObject* InMediaSourceProxy, int32 InIndex)
+	void SetMediaSourceProxy(const FMovieSceneObjectBindingID& InObjectBinding, int32 InIndex)
 	{
-		MediaSourceProxy = InMediaSourceProxy;
+		MediaSourceProxyBindingID = InObjectBinding;
 		MediaSourceProxyIndex = InIndex;
 	}
 
+	/**
+	 * If there is a proxy, get the media source from it.
+	 * Otherwise return our media source.
+	 * 
+	 * @param Player				Player used to get the proxy object from the binding.
+	 * @param SequenceID			ID used to get the proxy object from the binding.
+	 */
+	UMediaSource* GetMediaSourceOrProxy(IMovieScenePlayer& Player, FMovieSceneSequenceID SequenceID) const;
+
+	/**
+	 * Static version of GetMediaSourceOrProxy.
+	 *
+	 * @param Player				Player used to get the proxy object from the binding.
+	 * @param SequenceID			ID used to get the proxy object from the binding.
+	 */
+	static UMediaSource* GetMediaSourceOrProxy(IMovieScenePlayer& InPlayer, FMovieSceneSequenceID InSequenceID, UMediaSource* InMediaSource, const FMovieSceneObjectBindingID& InMediaSourceProxyBindingID, int32 InMediaSourceProxyIndex);
+
+
 #if WITH_EDITORONLY_DATA
 
 	/** @return The thumbnail reference frame offset from the start of this section */
@@ -145,4 +174,11 @@ public:
 	float ThumbnailReferenceOffset;
 
 #endif //WITH_EDITORONLY_DATA
+
+private:
+
+	/** The object to get the source to play from if you don't want to directly specify a media source. */
+	UPROPERTY(EditAnywhere, Category = "Section")
+	FMovieSceneObjectBindingID MediaSourceProxyBindingID;
+
 };
diff --git a/Engine/Plugins/Media/MediaCompositing/Source/MediaCompositing/Public/MovieSceneMediaTrack.h b/Engine/Plugins/Media/MediaCompositing/Source/MediaCompositing/Public/MovieSceneMediaTrack.h
--- a/Engine/Plugins/Media/MediaCompositing/Source/MediaCompositing/Public/MovieSceneMediaTrack.h
+++ b/Engine/Plugins/Media/MediaCompositing/Source/MediaCompositing/Public/MovieSceneMediaTrack.h
@@ -11,6 +11,7 @@
 
 class UMediaSource;
 class UMovieSceneMediaSection;
+struct FMovieSceneObjectBindingID;
 
 
 /**
@@ -37,14 +38,14 @@ public:
 	/** Adds a new media source to the track. */
 	virtual UMovieSceneSection* AddNewMediaSourceOnRow(UMediaSource& MediaSource, FFrameNumber Time, int32 RowIndex);
 	/** Adds a new media source to the track. */
-	virtual UMovieSceneSection* AddNewMediaSourceProxyOnRow(UObject& MediaSourceProxy, int32 MediaSourceProxyIndex, FFrameNumber Time, int32 RowIndex);
+	virtual UMovieSceneSection* AddNewMediaSourceProxyOnRow(const FMovieSceneObjectBindingID& ObjectBinding, int32 MediaSourceProxyIndex, FFrameNumber Time, int32 RowIndex);
 
 	/** Adds a new media source on the next available/non-overlapping row. */
 	virtual UMovieSceneSection* AddNewMediaSource(UMediaSource& MediaSource, FFrameNumber Time) { return AddNewMediaSourceOnRow(MediaSource, Time, INDEX_NONE); }
 	/** Adds a new media source on the next available/non-overlapping row. */
-	virtual UMovieSceneSection* AddNewMediaSourceProxy(UObject& MediaSourceProxy, int32 MediaSourceProxyIndex, FFrameNumber Time)
+	virtual UMovieSceneSection* AddNewMediaSourceProxy(const FMovieSceneObjectBindingID& ObjectBinding, int32 MediaSourceProxyIndex, FFrameNumber Time)
 	{
-		return AddNewMediaSourceProxyOnRow(MediaSourceProxy, MediaSourceProxyIndex, Time, INDEX_NONE); 
+		return AddNewMediaSourceProxyOnRow(ObjectBinding, MediaSourceProxyIndex, Time, INDEX_NONE);
 	}
 
 public:
diff --git a/Engine/Plugins/Media/MediaCompositing/Source/MediaCompositingEditor/Private/Sequencer/MediaThumbnailSection.cpp b/Engine/Plugins/Media/MediaCompositing/Source/MediaCompositingEditor/Private/Sequencer/MediaThumbnailSection.cpp
--- a/Engine/Plugins/Media/MediaCompositing/Source/MediaCompositingEditor/Private/Sequencer/MediaThumbnailSection.cpp
+++ b/Engine/Plugins/Media/MediaCompositing/Source/MediaCompositingEditor/Private/Sequencer/MediaThumbnailSection.cpp
@@ -72,7 +72,13 @@ float FMediaThumbnailSection::GetSectionHeight() const
 FText FMediaThumbnailSection::GetSectionTitle() const
 {
 	UMovieSceneMediaSection* MediaSection = CastChecked<UMovieSceneMediaSection>(Section);
-	UMediaSource* MediaSource = MediaSection->GetMediaSource();
+	TSharedPtr<ISequencer> Sequencer = SequencerPtr.Pin();
+	UMediaSource* MediaSource = nullptr;
+
+	if (MediaSection && Sequencer.IsValid())
+	{
+		MediaSource = MediaSection->GetMediaSourceOrProxy(*Sequencer, Sequencer->GetFocusedTemplateID());
+	}
 
 	if (MediaSource == nullptr)
 	{
diff --git a/Engine/Plugins/Media/MediaPlate/Source/MediaPlateEditor/Private/Sequencer/MediaPlateTrackEditor.cpp b/Engine/Plugins/Media/MediaPlate/Source/MediaPlateEditor/Private/Sequencer/MediaPlateTrackEditor.cpp
--- a/Engine/Plugins/Media/MediaPlate/Source/MediaPlateEditor/Private/Sequencer/MediaPlateTrackEditor.cpp
+++ b/Engine/Plugins/Media/MediaPlate/Source/MediaPlateEditor/Private/Sequencer/MediaPlateTrackEditor.cpp
@@ -12,6 +12,7 @@
 #include "MediaPlaylist.h"
 #include "Misc/MessageDialog.h"
 #include "Modules/ModuleManager.h"
+#include "MovieSceneObjectBindingID.h"
 #include "MovieSceneMediaSection.h"
 #include "MovieSceneMediaTrack.h"
 #include "SceneOutlinerModule.h"
@@ -159,13 +160,16 @@ void FMediaPlateTrackEditor::HandleActorAdded(AActor* Actor, FGuid TargetObjectG
 	{
 		if (UMediaPlateComponent* MediaPlateComponent = Actor->FindComponentByClass<UMediaPlateComponent>())
 		{
-			AddTrackForComponent(MediaPlateComponent);
+			AddTrackForComponent(MediaPlateComponent, TargetObjectGuid);
 		}
 	}
 }
 
-void FMediaPlateTrackEditor::AddTrackForComponent(UMediaPlateComponent* Component)
+void FMediaPlateTrackEditor::AddTrackForComponent(UMediaPlateComponent* Component, FGuid TargetObjectGuid)
 {
+	FMovieSceneObjectBindingID TargetObjectBindingID =
+		UE::MovieScene::FRelativeObjectBindingID(TargetObjectGuid);
+
 	// Get object.
 	UObject* Object = Component->GetOwner();
 	FFindOrCreateHandleResult HandleResult = FindOrCreateHandleToObject(Object);
@@ -188,7 +192,8 @@ void FMediaPlateTrackEditor::AddTrackForComponent(UMediaPlateComponent* Componen
 				UMediaSource* MediaSource = Playlist->Get(Index);
 				if (MediaSource != nullptr)
 				{
-					UMovieSceneSection* Section = MediaTrack->AddNewMediaSourceProxy(*Component, Index, FFrameNumber(0));
+					UMovieSceneSection* Section = MediaTrack->AddNewMediaSourceProxy(
+						TargetObjectBindingID, Index, FFrameNumber(0));
 					if (Section != nullptr)
 					{
 						// Start process to get the duration.
@@ -250,7 +255,7 @@ void FMediaPlateTrackEditor::ImportObjectBinding(const TArray<FGuid> ObjectBindi
 			if (UMediaPlateComponent* MediaPlateComponent = Actor->FindComponentByClass<UMediaPlateComponent>())
 			{
 				// Add tracks for this.
-				AddTrackForComponent(MediaPlateComponent);
+				AddTrackForComponent(MediaPlateComponent, ObjectBinding);
 			}
 		}
 	}
diff --git a/Engine/Plugins/Media/MediaPlate/Source/MediaPlateEditor/Private/Sequencer/MediaPlateTrackEditor.h b/Engine/Plugins/Media/MediaPlate/Source/MediaPlateEditor/Private/Sequencer/MediaPlateTrackEditor.h
--- a/Engine/Plugins/Media/MediaPlate/Source/MediaPlateEditor/Private/Sequencer/MediaPlateTrackEditor.h
+++ b/Engine/Plugins/Media/MediaPlate/Source/MediaPlateEditor/Private/Sequencer/MediaPlateTrackEditor.h
@@ -73,7 +73,7 @@ private:
 	/**
 	 * Adds a track for this component.
 	 */
-	void AddTrackForComponent(UMediaPlateComponent* Component);
+	void AddTrackForComponent(UMediaPlateComponent* Component, FGuid TargetObjectGuid);
 
 	/**
 	 * Adds tracks from an object binding.
