MediaCompositing: Fixed media tracks not playing the correct media when using a playlist and media plate.

diff --git a/Engine/Plugins/Media/MediaCompositing/Source/MediaCompositing/Private/MovieScene/MovieSceneMediaData.cpp b/Engine/Plugins/Media/MediaCompositing/Source/MediaCompositing/Private/MovieScene/MovieSceneMediaData.cpp
--- a/Engine/Plugins/Media/MediaCompositing/Source/MediaCompositing/Private/MovieScene/MovieSceneMediaData.cpp
+++ b/Engine/Plugins/Media/MediaCompositing/Source/MediaCompositing/Private/MovieScene/MovieSceneMediaData.cpp
@@ -28,7 +28,18 @@ FMovieSceneMediaData::~FMovieSceneMediaData()
 		// we are just going to open again soon anyway.
 		if ((PlayerProxy == nullptr) || (MediaPlayer->IsLooping() == false))
 		{
-			MediaPlayer->Close();
+			if (PlayerProxy != nullptr)
+			{
+				IMediaPlayerProxyInterface* PlayerProxyInterface = Cast<IMediaPlayerProxyInterface>(PlayerProxy);
+				if (PlayerProxyInterface != nullptr)
+				{
+					PlayerProxyInterface->ProxyClose();
+				}
+			}
+			else
+			{
+				MediaPlayer->Close();
+			}
 		}
 		MediaPlayer->RemoveFromRoot();
 	}
@@ -73,6 +84,12 @@ void FMovieSceneMediaData::Setup(UMediaPlayer* OverrideMediaPlayer, UObject* InP
 	if ((InPlayerProxy != nullptr) && (InPlayerProxy->Implements<UMediaPlayerProxyInterface>()))
 	{
 		PlayerProxy = InPlayerProxy;
+		IMediaPlayerProxyInterface* PlayerProxyInterface = Cast<IMediaPlayerProxyInterface>
+			(PlayerProxy);
+		if (PlayerProxyInterface != nullptr)
+		{
+			PlayerProxyInterface->ProxySetPlayOnOpen(false);
+		}
 	}
 	else
 	{
diff --git a/Engine/Plugins/Media/MediaCompositing/Source/MediaCompositing/Private/MovieScene/MovieSceneMediaTemplate.cpp b/Engine/Plugins/Media/MediaCompositing/Source/MediaCompositing/Private/MovieScene/MovieSceneMediaTemplate.cpp
--- a/Engine/Plugins/Media/MediaCompositing/Source/MediaCompositing/Private/MovieScene/MovieSceneMediaTemplate.cpp
+++ b/Engine/Plugins/Media/MediaCompositing/Source/MediaCompositing/Private/MovieScene/MovieSceneMediaTemplate.cpp
@@ -46,6 +46,16 @@ struct FMediaSectionBaseExecutionToken
 			BaseMediaSource, MediaSourceProxy, MediaSourceProxyIndex);
 	}
 
+	/**
+	 * Returns the index to identify the media source we are using in the proxy.
+	 */
+	int32 GetMediaSourceProxyIndex() const { return MediaSourceProxyIndex; }
+
+	/**
+	 * Tests if we have a media source proxy.
+	 */
+	bool IsMediaSourceProxyValid() const { return MediaSourceProxy.IsValid(); }
+
 private:
 	UMediaSource* BaseMediaSource;
 	FMovieSceneObjectBindingID MediaSourceProxy;
@@ -89,8 +99,17 @@ struct FMediaSectionPreRollExecutionToken
 			{
 				MediaSource->SetCacheSettings(PlayerProxyInterface->GetCacheSettings());
 			}
-			SectionData.SeekOnOpen(StartTime);
-			MediaPlayer->OpenSource(MediaSource);
+			if ((PlayerProxyInterface != nullptr) && (IsMediaSourceProxyValid()))
+			{
+				SectionData.SeekOnOpen(StartTime);
+				int32 PlaylistIndex = GetMediaSourceProxyIndex();
+				PlayerProxyInterface->ProxyOpenPlaylistIndex(PlaylistIndex);
+			}
+			else
+			{
+				SectionData.SeekOnOpen(StartTime);
+				MediaPlayer->OpenSource(MediaSource);
+			}
 		}
 	}
 
@@ -139,7 +158,21 @@ struct FMediaSectionExecutionToken
 		}
 
 		// open the media source if necessary
-		if (MediaPlayer->GetUrl().IsEmpty())
+		if ((PlayerProxyInterface != nullptr) && (IsMediaSourceProxyValid()))
+		{
+			int32 PlaylistIndex = GetMediaSourceProxyIndex();
+			if (PlayerProxyInterface->ProxyIsPlaylistIndexPlaying(PlaylistIndex) == false)
+			{
+				MediaSource->SetCacheSettings(PlayerProxyInterface->GetCacheSettings());
+				SectionData.SeekOnOpen(CurrentTime);
+				// Setup an initial blocking range - MediaFramework will block (even through the opening process) in its next tick...
+				PlayerProxyInterface->ProxyOpenPlaylistIndex(PlaylistIndex);
+				MediaPlayer->SetBlockOnTimeRange(TRange<FTimespan>(CurrentTime, CurrentTime + FrameDuration));
+
+				return;
+			}
+		}
+		else if (MediaPlayer->GetUrl().IsEmpty())
 		{
 			if (PlayerProxyInterface != nullptr)
 			{
diff --git a/Engine/Plugins/Media/MediaPlate/Source/MediaPlate/Private/MediaPlateComponent.cpp b/Engine/Plugins/Media/MediaPlate/Source/MediaPlate/Private/MediaPlateComponent.cpp
--- a/Engine/Plugins/Media/MediaPlate/Source/MediaPlate/Private/MediaPlateComponent.cpp
+++ b/Engine/Plugins/Media/MediaPlate/Source/MediaPlate/Private/MediaPlateComponent.cpp
@@ -217,9 +217,15 @@ UMediaTexture* UMediaPlateComponent::GetMediaTexture()
 }
 
 void UMediaPlateComponent::Open()
+{
+	ProxyOpenPlaylistIndex(0);
+}
+
+void UMediaPlateComponent::ProxyOpenPlaylistIndex(int32 Index)
 {
 	bIsMediaPlatePlaying = true;
 	CurrentRate = bPlayOnOpen ? 1.0f : 0.0f;
+	PlaylistIndex = Index;
 
 	if (IsVisible())
 	{
@@ -229,7 +235,7 @@ void UMediaPlateComponent::Open()
 			UMediaSource* MediaSource = nullptr;
 			if (MediaPlaylist != nullptr)
 			{
-				MediaSource = MediaPlaylist->Get(0);
+				MediaSource = MediaPlaylist->Get(Index);
 			}
 			bIsPlaying = PlayMediaSource(MediaSource, bPlayOnOpen);
 		}
@@ -593,6 +599,21 @@ UMediaSource* UMediaPlateComponent::GetMediaSourceFromIndex(int32 Index) const
 	return MediaSource;
 }
 
+void UMediaPlateComponent::ProxyClose()
+{
+	Close();
+}
+
+bool UMediaPlateComponent::ProxyIsPlaylistIndexPlaying(int32 Index) const
+{
+	return bIsMediaPlatePlaying && (Index == PlaylistIndex);
+}
+
+void UMediaPlateComponent::ProxySetPlayOnOpen(bool bInPlayOnOpen)
+{
+	bPlayOnOpen = bInPlayOnOpen;
+}
+
 void UMediaPlateComponent::RestartPlayer()
 {
 	if (MediaPlayer != nullptr)
diff --git a/Engine/Plugins/Media/MediaPlate/Source/MediaPlate/Public/MediaPlateComponent.h b/Engine/Plugins/Media/MediaPlate/Source/MediaPlate/Public/MediaPlateComponent.h
--- a/Engine/Plugins/Media/MediaPlate/Source/MediaPlate/Public/MediaPlateComponent.h
+++ b/Engine/Plugins/Media/MediaPlate/Source/MediaPlate/Public/MediaPlateComponent.h
@@ -242,6 +242,10 @@ public:
 	virtual bool IsExternalControlAllowed() override;
 	virtual const FMediaSourceCacheSettings& GetCacheSettings() const override;
 	virtual UMediaSource* GetMediaSourceFromIndex(int32 Index) const override;
+	virtual void ProxyClose() override;
+	virtual bool ProxyIsPlaylistIndexPlaying(int32 Index) const override;
+	virtual void ProxyOpenPlaylistIndex(int32 Index) override;
+	virtual void ProxySetPlayOnOpen(bool bInPlayOnOpen) override;
 
 private:
 	/**
diff --git a/Engine/Source/Runtime/MediaAssets/Public/MediaPlayerProxyInterface.h b/Engine/Source/Runtime/MediaAssets/Public/MediaPlayerProxyInterface.h
--- a/Engine/Source/Runtime/MediaAssets/Public/MediaPlayerProxyInterface.h
+++ b/Engine/Source/Runtime/MediaAssets/Public/MediaPlayerProxyInterface.h
@@ -56,5 +56,25 @@ public:
 	 * Get the media source for a given index.
 	 */
 	virtual UMediaSource* GetMediaSourceFromIndex(int32 Index) const = 0;
+
+	/**
+	 * Close the player.
+	 */
+	virtual void ProxyClose() = 0;
+
+	/**
+	 * Ask if a specific track in the playlist is playing.
+	 */
+	virtual bool ProxyIsPlaylistIndexPlaying(int32 Index) const = 0;
+
+	/**
+	 * Open a specific track in the playlist.
+	 */
+	virtual void ProxyOpenPlaylistIndex(int32 Index) = 0;
+
+	/**
+	 * Set the player to play on open.
+	 */
+	virtual void ProxySetPlayOnOpen(bool bInPlayOnOpen) = 0;
 	
 };
