MediaCompositing: Fixed media cache only being drawn for the first loop.

diff --git a/Engine/Plugins/Media/MediaCompositing/Source/MediaCompositingEditor/Private/Sequencer/MediaThumbnailSection.cpp b/Engine/Plugins/Media/MediaCompositing/Source/MediaCompositingEditor/Private/Sequencer/MediaThumbnailSection.cpp
--- a/Engine/Plugins/Media/MediaCompositing/Source/MediaCompositingEditor/Private/Sequencer/MediaThumbnailSection.cpp
+++ b/Engine/Plugins/Media/MediaCompositing/Source/MediaCompositingEditor/Private/Sequencer/MediaThumbnailSection.cpp
@@ -333,23 +333,29 @@ void FMediaThumbnailSection::DrawSampleStates(FSequencerSectionPainter& InPainte
 
 	TArray<TRange<FTimespan>> Ranges;
 	RangeSet.GetRanges(Ranges);
+	float LoopDrawOffset = -TimeToPixelConverter.SecondsToPixel(TickResolution.AsSeconds(MediaSection->StartFrameOffset));
 
-	for (auto& Range : Ranges)
+	while (LoopDrawOffset < SectionSize.X)
 	{
-		const float DrawOffset = FMath::RoundToNegativeInfinity(FTimespan::Ratio(Range.GetLowerBoundValue(), MediaDuration) * MediaSizeX) -
-			TimeToPixelConverter.SecondsToPixel(TickResolution.AsSeconds(MediaSection->StartFrameOffset)) +
-			TimeToPixelConverter.SecondsToPixel(0.0);
-		const float DrawSize = FMath::RoundToPositiveInfinity(FTimespan::Ratio(Range.Size<FTimespan>(), MediaDuration) * MediaSizeX);
-		const float BarHeight = 4.0f;
+		for (auto& Range : Ranges)
+		{
+			const float DrawOffset = FMath::RoundToNegativeInfinity(FTimespan::Ratio(Range.GetLowerBoundValue(), MediaDuration) * MediaSizeX) +	
+				LoopDrawOffset +
+				TimeToPixelConverter.SecondsToPixel(0.0);
+			const float DrawSize = FMath::RoundToPositiveInfinity(FTimespan::Ratio(Range.Size<FTimespan>(), MediaDuration) * MediaSizeX);
+			const float BarHeight = 4.0f;
+
+			FSlateDrawElement::MakeBox(
+				InPainter.DrawElements,
+				InPainter.LayerId++,
+				InPainter.SectionGeometry.ToPaintGeometry(FVector2D(DrawOffset, SectionSize.Y - BarHeight - 1.0f), FVector2D(DrawSize, BarHeight)),
+				GenericBrush,
+				ESlateDrawEffect::None,
+				Color
+			);
+		}
 
-		FSlateDrawElement::MakeBox(
-			InPainter.DrawElements,
-			InPainter.LayerId++,
-			InPainter.SectionGeometry.ToPaintGeometry(FVector2D(DrawOffset, SectionSize.Y - BarHeight - 1.0f), FVector2D(DrawSize, BarHeight)),
-			GenericBrush,
-			ESlateDrawEffect::None,
-			Color
-		);
+		LoopDrawOffset += MediaSizeX;
 	}
 }
 
