MediaCompositing: Uses its own media players again when dealing with proxies as this will be better for cross fades.

diff --git a/Engine/Plugins/Media/MediaCompositing/Source/MediaCompositing/Private/MovieScene/MovieSceneMediaData.cpp b/Engine/Plugins/Media/MediaCompositing/Source/MediaCompositing/Private/MovieScene/MovieSceneMediaData.cpp
--- a/Engine/Plugins/Media/MediaCompositing/Source/MediaCompositing/Private/MovieScene/MovieSceneMediaData.cpp
+++ b/Engine/Plugins/Media/MediaCompositing/Source/MediaCompositing/Private/MovieScene/MovieSceneMediaData.cpp
@@ -5,6 +5,7 @@
 #include "IMediaEventSink.h"
 #include "MediaPlayer.h"
 #include "MediaPlayerProxyInterface.h"
+#include "MediaTexture.h"
 #include "UObject/Package.h"
 #include "UObject/UObjectGlobals.h"
 
@@ -24,23 +25,7 @@ FMovieSceneMediaData::~FMovieSceneMediaData()
 	if (MediaPlayer != nullptr)
 	{
 		MediaPlayer->OnMediaEvent().RemoveAll(this);
-		// If we have a proxy and are looping, then don't bother closing as
-		// we are just going to open again soon anyway.
-		if ((PlayerProxy == nullptr) || (MediaPlayer->IsLooping() == false))
-		{
-			if (PlayerProxy != nullptr)
-			{
-				IMediaPlayerProxyInterface* PlayerProxyInterface = Cast<IMediaPlayerProxyInterface>(PlayerProxy);
-				if (PlayerProxyInterface != nullptr)
-				{
-					PlayerProxyInterface->ProxyClose();
-				}
-			}
-			else
-			{
-				MediaPlayer->Close();
-			}
-		}
+		MediaPlayer->Close();
 		MediaPlayer->RemoveFromRoot();
 	}
 }
@@ -79,17 +64,13 @@ void FMovieSceneMediaData::Setup(UMediaPlayer* OverrideMediaPlayer, UObject* InP
 	MediaPlayer->PlayOnOpen = false;
 	MediaPlayer->OnMediaEvent().AddRaw(this, &FMovieSceneMediaData::HandleMediaPlayerEvent);
 	MediaPlayer->AddToRoot();
+	ProxyMediaTexture.Reset();
 
 	// Do we have a valid proxy object?
 	if ((InPlayerProxy != nullptr) && (InPlayerProxy->Implements<UMediaPlayerProxyInterface>()))
 	{
 		PlayerProxy = InPlayerProxy;
-		IMediaPlayerProxyInterface* PlayerProxyInterface = Cast<IMediaPlayerProxyInterface>
-			(PlayerProxy);
-		if (PlayerProxyInterface != nullptr)
-		{
-			PlayerProxyInterface->ProxySetPlayOnOpen(false);
-		}
+		
 	}
 	else
 	{
@@ -97,6 +78,63 @@ void FMovieSceneMediaData::Setup(UMediaPlayer* OverrideMediaPlayer, UObject* InP
 	}
 }
 
+void FMovieSceneMediaData::Initialize(bool bIsEvaluating)
+{
+	if (bIsEvaluating)
+	{
+		AllocProxyMediaTexture();
+	}
+	else
+	{
+		DeallocProxyMediaTexture();
+	}
+}
+
+void FMovieSceneMediaData::TearDown()
+{
+	DeallocProxyMediaTexture();
+}
+
+void FMovieSceneMediaData::AllocProxyMediaTexture()
+{
+	if (PlayerProxy != nullptr)
+	{
+		if (ProxyMediaTexture == nullptr)
+		{
+			IMediaPlayerProxyInterface* PlayerProxyInterface = Cast<IMediaPlayerProxyInterface>(PlayerProxy);
+			if (PlayerProxyInterface != nullptr)
+			{
+				ProxyMediaTexture = PlayerProxyInterface->ProxyAllocMediaTexture();
+			}
+		}
+
+		if (ProxyMediaTexture != nullptr)
+		{
+			ProxyMediaTexture->SetMediaPlayer(MediaPlayer);
+		}
+	}
+}
+
+void FMovieSceneMediaData::DeallocProxyMediaTexture()
+{
+	if (PlayerProxy != nullptr)
+	{
+		if (ProxyMediaTexture != nullptr)
+		{
+			if (ProxyMediaTexture->GetMediaPlayer() == MediaPlayer)
+			{
+				ProxyMediaTexture->SetMediaPlayer(nullptr);
+			}
+		
+			IMediaPlayerProxyInterface* PlayerProxyInterface = Cast<IMediaPlayerProxyInterface>(PlayerProxy);
+			if (PlayerProxyInterface != nullptr)
+			{
+				PlayerProxyInterface->ProxyDeallocMediaTexture(ProxyMediaTexture.Get());
+			}
+			ProxyMediaTexture = nullptr;
+		}
+	}
+}
 
 /* FMediaSectionData callbacks
  *****************************************************************************/
diff --git a/Engine/Plugins/Media/MediaCompositing/Source/MediaCompositing/Private/MovieScene/MovieSceneMediaTemplate.cpp b/Engine/Plugins/Media/MediaCompositing/Source/MediaCompositing/Private/MovieScene/MovieSceneMediaTemplate.cpp
--- a/Engine/Plugins/Media/MediaCompositing/Source/MediaCompositing/Private/MovieScene/MovieSceneMediaTemplate.cpp
+++ b/Engine/Plugins/Media/MediaCompositing/Source/MediaCompositing/Private/MovieScene/MovieSceneMediaTemplate.cpp
@@ -99,17 +99,9 @@ struct FMediaSectionPreRollExecutionToken
 			{
 				MediaSource->SetCacheSettings(PlayerProxyInterface->GetCacheSettings());
 			}
-			if ((PlayerProxyInterface != nullptr) && (IsMediaSourceProxyValid()))
-			{
-				SectionData.SeekOnOpen(StartTime);
-				int32 PlaylistIndex = GetMediaSourceProxyIndex();
-				PlayerProxyInterface->ProxyOpenPlaylistIndex(PlaylistIndex);
-			}
-			else
-			{
-				SectionData.SeekOnOpen(StartTime);
-				MediaPlayer->OpenSource(MediaSource);
-			}
+			
+			SectionData.SeekOnOpen(StartTime);
+			MediaPlayer->OpenSource(MediaSource);
 		}
 	}
 
@@ -158,21 +150,7 @@ struct FMediaSectionExecutionToken
 		}
 
 		// open the media source if necessary
-		if ((PlayerProxyInterface != nullptr) && (IsMediaSourceProxyValid()))
-		{
-			int32 PlaylistIndex = GetMediaSourceProxyIndex();
-			if (PlayerProxyInterface->ProxyIsPlaylistIndexPlaying(PlaylistIndex) == false)
-			{
-				MediaSource->SetCacheSettings(PlayerProxyInterface->GetCacheSettings());
-				SectionData.SeekOnOpen(CurrentTime);
-				// Setup an initial blocking range - MediaFramework will block (even through the opening process) in its next tick...
-				PlayerProxyInterface->ProxyOpenPlaylistIndex(PlaylistIndex);
-				MediaPlayer->SetBlockOnTimeRange(TRange<FTimespan>(CurrentTime, CurrentTime + FrameDuration));
-
-				return;
-			}
-		}
-		else if (MediaPlayer->GetUrl().IsEmpty())
+		if (MediaPlayer->GetUrl().IsEmpty())
 		{
 			if (PlayerProxyInterface != nullptr)
 			{
@@ -232,8 +210,7 @@ struct FMediaSectionExecutionToken
 				// Set rate
 				// (note that the DIRECTION is important, but the magnitude is not - as we use blocked playback, the range setup to block on will serve as external clock to the player,
 				//  the direction is taken into account as hint for internal operation of the player)
-				if (!SetRate(MediaPlayer, PlayerProxyInterface,
-					(Context.GetDirection() == EPlayDirection::Forwards) ? 1.0f : -1.0f))
+				if (!MediaPlayer->SetRate((Context.GetDirection() == EPlayDirection::Forwards) ? 1.0f : -1.0f))
 				{
 					// Failed to set needed rate: better switch off blocking and bail...
 					MediaPlayer->SetBlockOnTimeRange(TRange<FTimespan>::Empty());
@@ -250,7 +227,7 @@ struct FMediaSectionExecutionToken
 				float CurrentPlayerRate = MediaPlayer->GetRate();
 				if (Context.GetDirection() == EPlayDirection::Forwards && CurrentPlayerRate < 0.0f)
 				{
-					if (!SetRate(MediaPlayer, PlayerProxyInterface, 1.0f))
+					if (!MediaPlayer->SetRate(1.0f))
 					{
 						// Failed to set needed rate: better switch off blocking and bail...
 						MediaPlayer->SetBlockOnTimeRange(TRange<FTimespan>::Empty());
@@ -259,7 +236,7 @@ struct FMediaSectionExecutionToken
 				}
 				else if (Context.GetDirection() == EPlayDirection::Backwards && CurrentPlayerRate > 0.0f)
 				{
-					if (!SetRate(MediaPlayer, PlayerProxyInterface, -1.0f))
+					if (!MediaPlayer->SetRate(-1.0f))
 					{
 						// Failed to set needed rate: better switch off blocking and bail...
 						MediaPlayer->SetBlockOnTimeRange(TRange<FTimespan>::Empty());
@@ -272,7 +249,7 @@ struct FMediaSectionExecutionToken
 		{
 			if (MediaPlayer->IsPlaying())
 			{
-				SetRate(MediaPlayer, PlayerProxyInterface, 0.0f);
+				MediaPlayer->SetRate(0.0f);
 			}
 
 			MediaPlayer->Seek(MediaTime);
@@ -285,22 +262,6 @@ struct FMediaSectionExecutionToken
 
 private:
 
-	/**
-	 * Call this to set the rate.
-	 * It will use the proxy if available, otherwise it will use the player.
-	 */
-	bool SetRate(UMediaPlayer* Player, IMediaPlayerProxyInterface* Proxy, float Rate)
-	{
-		if (Proxy != nullptr)
-		{
-			return Proxy->SetProxyRate(Rate);
-		}
-		else
-		{
-			return Player->SetRate(Rate);
-		}
-	}
-
 	FTimespan CurrentTime;
 	FTimespan FrameDuration;
 	float PlaybackRate;
@@ -417,7 +378,7 @@ void FMovieSceneMediaSectionTemplate::Initialize(const FMovieSceneEvaluationOper
 						UObject* BoundObject = WeakObject.Get();
 						if (BoundObject != nullptr)
 						{
-							MediaPlayer = MediaAssetsModule->GetPlayerFromObject(BoundObject, PlayerProxy);
+							MediaAssetsModule->GetPlayerFromObject(BoundObject, PlayerProxy);
 							break;
 						}
 					}
@@ -443,6 +404,7 @@ void FMovieSceneMediaSectionTemplate::Initialize(const FMovieSceneEvaluationOper
 	}
 
 	const bool IsEvaluating = !(Context.IsPreRoll() || Context.IsPostRoll() || (Context.GetTime().FrameNumber >= Params.SectionEndFrame));
+	SectionData->Initialize(IsEvaluating);
 
 	if (Params.MediaSoundComponent != nullptr)
 	{
@@ -519,5 +481,7 @@ void FMovieSceneMediaSectionTemplate::TearDown(FPersistentEvaluationData& Persis
 	{
 		Params.MediaTexture->SetMediaPlayer(nullptr);
 	}
+
+	SectionData->TearDown();
 }
 
diff --git a/Engine/Plugins/Media/MediaCompositing/Source/MediaCompositing/Public/MovieSceneMediaData.h b/Engine/Plugins/Media/MediaCompositing/Source/MediaCompositing/Public/MovieSceneMediaData.h
--- a/Engine/Plugins/Media/MediaCompositing/Source/MediaCompositing/Public/MovieSceneMediaData.h
+++ b/Engine/Plugins/Media/MediaCompositing/Source/MediaCompositing/Public/MovieSceneMediaData.h
@@ -8,6 +8,7 @@
 
 class IMediaPlayerProxy;
 class UMediaPlayer;
+class UMediaTexture;
 enum class EMediaEvent;
 
 
@@ -47,7 +48,26 @@ public:
 	/** Set up this persistent data object. */
 	void Setup(UMediaPlayer* OverrideMediaPlayer, UObject* InPlayerProxy);
 
+	/**
+	 * Called from FMovieSceneMediaSectionTemplate::Initialize.
+	 */
+	void Initialize(bool bIsEvaluating);
+
+	/**
+	 * Called from FMovieSceneMediaSectionTemplate::TearDown.
+	 */
+	void TearDown();
+
 private:
+	/**
+	 * Sets up our proxy media texture.
+	 */
+	void AllocProxyMediaTexture();
+
+	/**
+	 * Removes our proxy media texture.
+	 */
+	void DeallocProxyMediaTexture();
 
 	/** Callback for media player events. */
 	void HandleMediaPlayerEvent(EMediaEvent Event);
@@ -59,6 +79,8 @@ private:
 	UMediaPlayer* MediaPlayer;
 	/** Optional proxy for the media player. */
 	TWeakObjectPtr<UObject> PlayerProxy;
+	/** Media texture allocated from the proxy. */
+	TWeakObjectPtr<UMediaTexture> ProxyMediaTexture;
 
 	/** The time to seek to after the media source is opened. */
 	FTimespan SeekOnOpenTime;
diff --git a/Engine/Plugins/Media/MediaPlate/Source/MediaPlate/Private/MediaPlateComponent.cpp b/Engine/Plugins/Media/MediaPlate/Source/MediaPlate/Private/MediaPlateComponent.cpp
--- a/Engine/Plugins/Media/MediaPlate/Source/MediaPlate/Private/MediaPlateComponent.cpp
+++ b/Engine/Plugins/Media/MediaPlate/Source/MediaPlate/Private/MediaPlateComponent.cpp
@@ -599,11 +599,20 @@ UMediaSource* UMediaPlateComponent::GetMediaSourceFromIndex(int32 Index) const
 	return MediaSource;
 }
 
+UMediaTexture* UMediaPlateComponent::ProxyAllocMediaTexture()
+{
+	return MediaTexture;
+}
+
 void UMediaPlateComponent::ProxyClose()
 {
 	Close();
 }
 
+void UMediaPlateComponent::ProxyDeallocMediaTexture(UMediaTexture* InMediaTexture)
+{
+}
+
 bool UMediaPlateComponent::ProxyIsPlaylistIndexPlaying(int32 Index) const
 {
 	return bIsMediaPlatePlaying && (Index == PlaylistIndex);
diff --git a/Engine/Plugins/Media/MediaPlate/Source/MediaPlate/Public/MediaPlateComponent.h b/Engine/Plugins/Media/MediaPlate/Source/MediaPlate/Public/MediaPlateComponent.h
--- a/Engine/Plugins/Media/MediaPlate/Source/MediaPlate/Public/MediaPlateComponent.h
+++ b/Engine/Plugins/Media/MediaPlate/Source/MediaPlate/Public/MediaPlateComponent.h
@@ -242,7 +242,9 @@ public:
 	virtual bool IsExternalControlAllowed() override;
 	virtual const FMediaSourceCacheSettings& GetCacheSettings() const override;
 	virtual UMediaSource* GetMediaSourceFromIndex(int32 Index) const override;
+	virtual UMediaTexture* ProxyAllocMediaTexture() override;
 	virtual void ProxyClose() override;
+	virtual void ProxyDeallocMediaTexture(UMediaTexture* MediaTexture) override;
 	virtual bool ProxyIsPlaylistIndexPlaying(int32 Index) const override;
 	virtual void ProxyOpenPlaylistIndex(int32 Index) override;
 	virtual void ProxySetPlayOnOpen(bool bInPlayOnOpen) override;
diff --git a/Engine/Source/Runtime/MediaAssets/Public/MediaPlayerProxyInterface.h b/Engine/Source/Runtime/MediaAssets/Public/MediaPlayerProxyInterface.h
--- a/Engine/Source/Runtime/MediaAssets/Public/MediaPlayerProxyInterface.h
+++ b/Engine/Source/Runtime/MediaAssets/Public/MediaPlayerProxyInterface.h
@@ -6,7 +6,8 @@
 
 #include "MediaPlayerProxyInterface.generated.h"
 
-class UMediaSource; 
+class UMediaSource;
+class UMediaTexture;
 struct FMediaSourceCacheSettings;
 
 /**
@@ -57,11 +58,21 @@ public:
 	 */
 	virtual UMediaSource* GetMediaSourceFromIndex(int32 Index) const = 0;
 
+	/**
+	 * Get a media texture that we can assign a media player to.
+	 */
+	virtual UMediaTexture* ProxyAllocMediaTexture() = 0;
+
 	/**
 	 * Close the player.
 	 */
 	virtual void ProxyClose() = 0;
 
+	/**
+	 * Return a media texture received from ProxyAllocMediaTexture.
+	 */
+	virtual void ProxyDeallocMediaTexture(UMediaTexture* InMediaTexture) = 0;
+
 	/**
 	 * Ask if a specific track in the playlist is playing.
 	 */
