MediaCompositing: Added support for a texture per section when using a proxy.

diff --git a/Engine/Plugins/Media/MediaCompositing/Source/MediaCompositing/Private/MovieScene/MovieSceneMediaData.cpp b/Engine/Plugins/Media/MediaCompositing/Source/MediaCompositing/Private/MovieScene/MovieSceneMediaData.cpp
--- a/Engine/Plugins/Media/MediaCompositing/Source/MediaCompositing/Private/MovieScene/MovieSceneMediaData.cpp
+++ b/Engine/Plugins/Media/MediaCompositing/Source/MediaCompositing/Private/MovieScene/MovieSceneMediaData.cpp
@@ -46,7 +46,7 @@ void FMovieSceneMediaData::SeekOnOpen(FTimespan Time)
 }
 
 
-void FMovieSceneMediaData::Setup(UMediaPlayer* OverrideMediaPlayer, UObject* InPlayerProxy)
+void FMovieSceneMediaData::Setup(UMediaPlayer* OverrideMediaPlayer, UObject* InPlayerProxy, int32 ProxyTextureIndex)
 {
 	// Ensure we don't already have a media player set. Setup should only be called once
 	check(!MediaPlayer);
@@ -70,7 +70,11 @@ void FMovieSceneMediaData::Setup(UMediaPlayer* OverrideMediaPlayer, UObject* InP
 	if ((InPlayerProxy != nullptr) && (InPlayerProxy->Implements<UMediaPlayerProxyInterface>()))
 	{
 		PlayerProxy = InPlayerProxy;
-		
+		IMediaPlayerProxyInterface* PlayerProxyInterface = Cast<IMediaPlayerProxyInterface>(PlayerProxy);
+		if (PlayerProxyInterface != nullptr)
+		{
+			ProxyMediaTexture = PlayerProxyInterface->ProxyGetMediaTexture(ProxyTextureIndex);
+		}
 	}
 	else
 	{
@@ -82,32 +86,23 @@ void FMovieSceneMediaData::Initialize(bool bIsEvaluating)
 {
 	if (bIsEvaluating)
 	{
-		AllocProxyMediaTexture();
+		StartUsingProxyMediaTexture();
 	}
 	else
 	{
-		DeallocProxyMediaTexture();
+		StopUsingProxyMediaTexture();
 	}
 }
 
 void FMovieSceneMediaData::TearDown()
 {
-	DeallocProxyMediaTexture();
+	StopUsingProxyMediaTexture();
 }
 
-void FMovieSceneMediaData::AllocProxyMediaTexture()
+void FMovieSceneMediaData::StartUsingProxyMediaTexture()
 {
 	if (PlayerProxy != nullptr)
 	{
-		if (ProxyMediaTexture == nullptr)
-		{
-			IMediaPlayerProxyInterface* PlayerProxyInterface = Cast<IMediaPlayerProxyInterface>(PlayerProxy);
-			if (PlayerProxyInterface != nullptr)
-			{
-				ProxyMediaTexture = PlayerProxyInterface->ProxyAllocMediaTexture();
-			}
-		}
-
 		if (ProxyMediaTexture != nullptr)
 		{
 			ProxyMediaTexture->SetMediaPlayer(MediaPlayer);
@@ -115,7 +110,7 @@ void FMovieSceneMediaData::AllocProxyMediaTexture()
 	}
 }
 
-void FMovieSceneMediaData::DeallocProxyMediaTexture()
+void FMovieSceneMediaData::StopUsingProxyMediaTexture()
 {
 	if (PlayerProxy != nullptr)
 	{
@@ -125,13 +120,6 @@ void FMovieSceneMediaData::DeallocProxyMediaTexture()
 			{
 				ProxyMediaTexture->SetMediaPlayer(nullptr);
 			}
-		
-			IMediaPlayerProxyInterface* PlayerProxyInterface = Cast<IMediaPlayerProxyInterface>(PlayerProxy);
-			if (PlayerProxyInterface != nullptr)
-			{
-				PlayerProxyInterface->ProxyDeallocMediaTexture(ProxyMediaTexture.Get());
-			}
-			ProxyMediaTexture = nullptr;
 		}
 	}
 }
diff --git a/Engine/Plugins/Media/MediaCompositing/Source/MediaCompositing/Private/MovieScene/MovieSceneMediaSection.cpp b/Engine/Plugins/Media/MediaCompositing/Source/MediaCompositing/Private/MovieScene/MovieSceneMediaSection.cpp
--- a/Engine/Plugins/Media/MediaCompositing/Source/MediaCompositing/Private/MovieScene/MovieSceneMediaSection.cpp
+++ b/Engine/Plugins/Media/MediaCompositing/Source/MediaCompositing/Private/MovieScene/MovieSceneMediaSection.cpp
@@ -27,6 +27,7 @@ UMovieSceneMediaSection::UMovieSceneMediaSection(const FObjectInitializer& Objec
 	: Super(ObjectInitializer)
 	, MediaSourceProxyIndex(0)
 	, bLooping(true)
+	, TextureIndex(0)
 	, bHasMediaPlayerProxy(false)
 {
 #if WITH_EDITORONLY_DATA
diff --git a/Engine/Plugins/Media/MediaCompositing/Source/MediaCompositing/Private/MovieScene/MovieSceneMediaTemplate.cpp b/Engine/Plugins/Media/MediaCompositing/Source/MediaCompositing/Private/MovieScene/MovieSceneMediaTemplate.cpp
--- a/Engine/Plugins/Media/MediaCompositing/Source/MediaCompositing/Private/MovieScene/MovieSceneMediaTemplate.cpp
+++ b/Engine/Plugins/Media/MediaCompositing/Source/MediaCompositing/Private/MovieScene/MovieSceneMediaTemplate.cpp
@@ -296,6 +296,8 @@ FMovieSceneMediaSectionTemplate::FMovieSceneMediaSectionTemplate(const UMovieSce
 	{
 		Params.SectionEndFrame = InSection.GetRange().GetUpperBoundValue();
 	}
+
+	Params.ProxyTextureIndex = InSection.TextureIndex;
 }
 
 
@@ -388,7 +390,7 @@ void FMovieSceneMediaSectionTemplate::Initialize(const FMovieSceneEvaluationOper
 
 		// Add section data.
 		SectionData = &PersistentData.AddSectionData<FMovieSceneMediaData>();
-		SectionData->Setup(MediaPlayer, PlayerProxy);
+		SectionData->Setup(MediaPlayer, PlayerProxy, Params.ProxyTextureIndex);
 	}
 
 	if (!ensure(SectionData != nullptr))
diff --git a/Engine/Plugins/Media/MediaCompositing/Source/MediaCompositing/Private/MovieScene/MovieSceneMediaTemplate.h b/Engine/Plugins/Media/MediaCompositing/Source/MediaCompositing/Private/MovieScene/MovieSceneMediaTemplate.h
--- a/Engine/Plugins/Media/MediaCompositing/Source/MediaCompositing/Private/MovieScene/MovieSceneMediaTemplate.h
+++ b/Engine/Plugins/Media/MediaCompositing/Source/MediaCompositing/Private/MovieScene/MovieSceneMediaTemplate.h
@@ -51,6 +51,9 @@ struct FMovieSceneMediaSectionParams
 	UPROPERTY()
 	FFrameNumber StartFrameOffset;
 
+	UPROPERTY()
+	int32 ProxyTextureIndex;
+
 	FMovieSceneMediaSectionParams()
 		: MediaSoundComponent(nullptr)
 		, MediaSource(nullptr)
@@ -58,6 +61,7 @@ struct FMovieSceneMediaSectionParams
 		, MediaTexture(nullptr)
 		, MediaPlayer(nullptr)
 		, bLooping(false)
+		, ProxyTextureIndex(0)
 	{}
 };
 
diff --git a/Engine/Plugins/Media/MediaCompositing/Source/MediaCompositing/Public/MovieSceneMediaData.h b/Engine/Plugins/Media/MediaCompositing/Source/MediaCompositing/Public/MovieSceneMediaData.h
--- a/Engine/Plugins/Media/MediaCompositing/Source/MediaCompositing/Public/MovieSceneMediaData.h
+++ b/Engine/Plugins/Media/MediaCompositing/Source/MediaCompositing/Public/MovieSceneMediaData.h
@@ -46,7 +46,7 @@ public:
 	void SeekOnOpen(FTimespan Time);
 
 	/** Set up this persistent data object. */
-	void Setup(UMediaPlayer* OverrideMediaPlayer, UObject* InPlayerProxy);
+	void Setup(UMediaPlayer* OverrideMediaPlayer, UObject* InPlayerProxy, int32 ProxyTextureIndex);
 
 	/**
 	 * Called from FMovieSceneMediaSectionTemplate::Initialize.
@@ -60,14 +60,14 @@ public:
 
 private:
 	/**
-	 * Sets up our proxy media texture.
+	 * Does the work needed so we can use our proxy media texture.
 	 */
-	void AllocProxyMediaTexture();
+	void StartUsingProxyMediaTexture();
 
 	/**
-	 * Removes our proxy media texture.
+	 * Does the work needed when we no longer use our proxy media texture.
 	 */
-	void DeallocProxyMediaTexture();
+	void StopUsingProxyMediaTexture();
 
 	/** Callback for media player events. */
 	void HandleMediaPlayerEvent(EMediaEvent Event);
diff --git a/Engine/Plugins/Media/MediaCompositing/Source/MediaCompositing/Public/MovieSceneMediaSection.h b/Engine/Plugins/Media/MediaCompositing/Source/MediaCompositing/Public/MovieSceneMediaSection.h
--- a/Engine/Plugins/Media/MediaCompositing/Source/MediaCompositing/Public/MovieSceneMediaSection.h
+++ b/Engine/Plugins/Media/MediaCompositing/Source/MediaCompositing/Public/MovieSceneMediaSection.h
@@ -65,6 +65,10 @@ public:
 	UPROPERTY(EditAnywhere, BlueprintReadWrite, Category = "Media", meta = (EditCondition = "!bHasMediaPlayerProxy", HideEditConditionToggle = true, EditConditionHides), AdvancedDisplay)
 	FMediaSourceCacheSettings CacheSettings;
 
+	/** If using an object like a MediaPlate, then this determines which texture to use for crossfading purposes. */
+	UPROPERTY(EditAnywhere, BlueprintReadWrite, Category = "Media", meta = (EditCondition = "bHasMediaPlayerProxy", HideEditConditionToggle = true, EditConditionHides), AdvancedDisplay)
+	int32 TextureIndex;
+
 	/** True if the object bound to this track has a media player proxy. */
 	UPROPERTY(Transient)
 	bool bHasMediaPlayerProxy;
diff --git a/Engine/Plugins/Media/MediaPlate/Source/MediaPlate/Private/MediaPlateComponent.cpp b/Engine/Plugins/Media/MediaPlate/Source/MediaPlate/Private/MediaPlateComponent.cpp
--- a/Engine/Plugins/Media/MediaPlate/Source/MediaPlate/Private/MediaPlateComponent.cpp
+++ b/Engine/Plugins/Media/MediaPlate/Source/MediaPlate/Private/MediaPlateComponent.cpp
@@ -599,18 +599,15 @@ UMediaSource* UMediaPlateComponent::GetMediaSourceFromIndex(int32 Index) const
 	return MediaSource;
 }
 
-UMediaTexture* UMediaPlateComponent::ProxyAllocMediaTexture()
-{
-	return MediaTexture;
-}
-
 void UMediaPlateComponent::ProxyClose()
 {
 	Close();
 }
 
-void UMediaPlateComponent::ProxyDeallocMediaTexture(UMediaTexture* InMediaTexture)
+UMediaTexture* UMediaPlateComponent::ProxyGetMediaTexture(int32 TextureIndex)
 {
+	UMediaTexture* OutTexture = MediaTexture;
+	return OutTexture;
 }
 
 bool UMediaPlateComponent::ProxyIsPlaylistIndexPlaying(int32 Index) const
diff --git a/Engine/Plugins/Media/MediaPlate/Source/MediaPlate/Public/MediaPlateComponent.h b/Engine/Plugins/Media/MediaPlate/Source/MediaPlate/Public/MediaPlateComponent.h
--- a/Engine/Plugins/Media/MediaPlate/Source/MediaPlate/Public/MediaPlateComponent.h
+++ b/Engine/Plugins/Media/MediaPlate/Source/MediaPlate/Public/MediaPlateComponent.h
@@ -242,9 +242,8 @@ public:
 	virtual bool IsExternalControlAllowed() override;
 	virtual const FMediaSourceCacheSettings& GetCacheSettings() const override;
 	virtual UMediaSource* GetMediaSourceFromIndex(int32 Index) const override;
-	virtual UMediaTexture* ProxyAllocMediaTexture() override;
 	virtual void ProxyClose() override;
-	virtual void ProxyDeallocMediaTexture(UMediaTexture* MediaTexture) override;
+	virtual UMediaTexture* ProxyGetMediaTexture(int32 TextureIndex) override;
 	virtual bool ProxyIsPlaylistIndexPlaying(int32 Index) const override;
 	virtual void ProxyOpenPlaylistIndex(int32 Index) override;
 	virtual void ProxySetPlayOnOpen(bool bInPlayOnOpen) override;
diff --git a/Engine/Source/Runtime/MediaAssets/Public/MediaPlayerProxyInterface.h b/Engine/Source/Runtime/MediaAssets/Public/MediaPlayerProxyInterface.h
--- a/Engine/Source/Runtime/MediaAssets/Public/MediaPlayerProxyInterface.h
+++ b/Engine/Source/Runtime/MediaAssets/Public/MediaPlayerProxyInterface.h
@@ -58,20 +58,15 @@ public:
 	 */
 	virtual UMediaSource* GetMediaSourceFromIndex(int32 Index) const = 0;
 
-	/**
-	 * Get a media texture that we can assign a media player to.
-	 */
-	virtual UMediaTexture* ProxyAllocMediaTexture() = 0;
-
 	/**
 	 * Close the player.
 	 */
 	virtual void ProxyClose() = 0;
 
 	/**
-	 * Return a media texture received from ProxyAllocMediaTexture.
+	 * Get a media texture that we can assign a media player to.
 	 */
-	virtual void ProxyDeallocMediaTexture(UMediaTexture* InMediaTexture) = 0;
+	virtual UMediaTexture* ProxyGetMediaTexture(int32 TextureIndex) = 0;
 
 	/**
 	 * Ask if a specific track in the playlist is playing.
