MediaCompositing: Can key the texture blend for each section for cross fading.

diff --git a/Engine/Plugins/Media/MediaCompositing/Source/MediaCompositing/Private/MovieScene/MovieSceneMediaData.cpp b/Engine/Plugins/Media/MediaCompositing/Source/MediaCompositing/Private/MovieScene/MovieSceneMediaData.cpp
--- a/Engine/Plugins/Media/MediaCompositing/Source/MediaCompositing/Private/MovieScene/MovieSceneMediaData.cpp
+++ b/Engine/Plugins/Media/MediaCompositing/Source/MediaCompositing/Private/MovieScene/MovieSceneMediaData.cpp
@@ -16,6 +16,7 @@
 FMovieSceneMediaData::FMovieSceneMediaData()
 	: bOverrideMediaPlayer(false)
 	, MediaPlayer(nullptr)
+	, ProxyTextureIndex(0)
 	, SeekOnOpenTime(FTimespan::MinValue())
 { }
 
@@ -46,7 +47,7 @@ void FMovieSceneMediaData::SeekOnOpen(FTimespan Time)
 }
 
 
-void FMovieSceneMediaData::Setup(UMediaPlayer* OverrideMediaPlayer, UObject* InPlayerProxy, int32 ProxyTextureIndex)
+void FMovieSceneMediaData::Setup(UMediaPlayer* OverrideMediaPlayer, UObject* InPlayerProxy, int32 InProxyTextureIndex)
 {
 	// Ensure we don't already have a media player set. Setup should only be called once
 	check(!MediaPlayer);
@@ -65,6 +66,7 @@ void FMovieSceneMediaData::Setup(UMediaPlayer* OverrideMediaPlayer, UObject* InP
 	MediaPlayer->OnMediaEvent().AddRaw(this, &FMovieSceneMediaData::HandleMediaPlayerEvent);
 	MediaPlayer->AddToRoot();
 	ProxyMediaTexture.Reset();
+	ProxyTextureIndex = InProxyTextureIndex;
 
 	// Do we have a valid proxy object?
 	if ((InPlayerProxy != nullptr) && (InPlayerProxy->Implements<UMediaPlayerProxyInterface>()))
diff --git a/Engine/Plugins/Media/MediaCompositing/Source/MediaCompositing/Private/MovieScene/MovieSceneMediaSection.cpp b/Engine/Plugins/Media/MediaCompositing/Source/MediaCompositing/Private/MovieScene/MovieSceneMediaSection.cpp
--- a/Engine/Plugins/Media/MediaCompositing/Source/MediaCompositing/Private/MovieScene/MovieSceneMediaSection.cpp
+++ b/Engine/Plugins/Media/MediaCompositing/Source/MediaCompositing/Private/MovieScene/MovieSceneMediaSection.cpp
@@ -35,6 +35,17 @@ UMovieSceneMediaSection::UMovieSceneMediaSection(const FObjectInitializer& Objec
 #endif
 
 	EvalOptions.CompletionMode = EMovieSceneCompletionMode::RestoreState;
+	TextureBlend.SetDefault(1.0f);
+
+#if WITH_EDITOR
+
+	ChannelProxy = MakeShared<FMovieSceneChannelProxy>(TextureBlend, FMovieSceneChannelMetaData(), TMovieSceneExternalValue<float>::Make());
+
+#else
+
+	ChannelProxy = MakeShared<FMovieSceneChannelProxy>(TextureBlend);
+
+#endif
 }
 
 void UMovieSceneMediaSection::PostInitProperties()
diff --git a/Engine/Plugins/Media/MediaCompositing/Source/MediaCompositing/Private/MovieScene/MovieSceneMediaTemplate.cpp b/Engine/Plugins/Media/MediaCompositing/Source/MediaCompositing/Private/MovieScene/MovieSceneMediaTemplate.cpp
--- a/Engine/Plugins/Media/MediaCompositing/Source/MediaCompositing/Private/MovieScene/MovieSceneMediaTemplate.cpp
+++ b/Engine/Plugins/Media/MediaCompositing/Source/MediaCompositing/Private/MovieScene/MovieSceneMediaTemplate.cpp
@@ -114,11 +114,12 @@ private:
 struct FMediaSectionExecutionToken
 	: FMediaSectionBaseExecutionToken
 {
-	FMediaSectionExecutionToken(UMediaSource* InMediaSource, FMovieSceneObjectBindingID InMediaSourceProxy, int32 InMediaSourceProxyIndex, FTimespan InCurrentTime, FTimespan InFrameDuration)
+	FMediaSectionExecutionToken(UMediaSource* InMediaSource, FMovieSceneObjectBindingID InMediaSourceProxy, int32 InMediaSourceProxyIndex, float InProxyTextureBlend, FTimespan InCurrentTime, FTimespan InFrameDuration)
 		: FMediaSectionBaseExecutionToken(InMediaSource, InMediaSourceProxy, InMediaSourceProxyIndex)
 		, CurrentTime(InCurrentTime)
 		, FrameDuration(InFrameDuration)
 		, PlaybackRate(1.0f)
+		, ProxyTextureBlend(InProxyTextureBlend)
 	{ }
 
 	virtual void Execute(const FMovieSceneContext& Context, const FMovieSceneEvaluationOperand& Operand, FPersistentEvaluationData& PersistentData, IMovieScenePlayer& Player) override
@@ -141,6 +142,7 @@ struct FMediaSectionExecutionToken
 				(PlayerProxy);
 			if (PlayerProxyInterface != nullptr)
 			{
+				PlayerProxyInterface->ProxySetTextureBlend(SectionData.GetProxyTextureIndex(), ProxyTextureBlend);
 				// Can we control the player?
 				if (PlayerProxyInterface->IsExternalControlAllowed() == false)
 				{
@@ -265,6 +267,7 @@ private:
 	FTimespan CurrentTime;
 	FTimespan FrameDuration;
 	float PlaybackRate;
+	float ProxyTextureBlend;
 };
 
 
@@ -297,6 +300,7 @@ FMovieSceneMediaSectionTemplate::FMovieSceneMediaSectionTemplate(const UMovieSce
 		Params.SectionEndFrame = InSection.GetRange().GetUpperBoundValue();
 	}
 
+	Params.ProxyTextureBlend = InSection.TextureBlend;
 	Params.ProxyTextureIndex = InSection.TextureIndex;
 }
 
@@ -336,6 +340,9 @@ void FMovieSceneMediaSectionTemplate::Evaluate(const FMovieSceneEvaluationOperan
 		const double FrameDurationInSeconds = FMath::Max(FrameRate.AsSeconds(FFrameTime(1)), (Context.GetRange().Size<FFrameTime>()) / Context.GetFrameRate());
 		const int64 FrameDurationTicks = FrameDurationInSeconds * ETimespan::TicksPerSecond;
 
+		float ProxyTextureBlend = 1.0f;
+		Params.ProxyTextureBlend.Evaluate(Context.GetTime(), ProxyTextureBlend);
+
 		#if MOVIESCENEMEDIATEMPLATE_TRACE_EVALUATION
 			GLog->Logf(ELogVerbosity::Log, TEXT("Evaluating frame %i+%f, FrameRate %i/%i, FrameTicks %d, FrameDurationTicks %d"),
 				Context.GetTime().GetFrame().Value,
@@ -347,7 +354,7 @@ void FMovieSceneMediaSectionTemplate::Evaluate(const FMovieSceneEvaluationOperan
 			);
 		#endif
 
-		ExecutionTokens.Add(FMediaSectionExecutionToken(MediaSource, Params.MediaSourceProxy, Params.MediaSourceProxyIndex, FTimespan(FrameTicks), FTimespan(FrameDurationTicks)));
+		ExecutionTokens.Add(FMediaSectionExecutionToken(MediaSource, Params.MediaSourceProxy, Params.MediaSourceProxyIndex, ProxyTextureBlend, FTimespan(FrameTicks), FTimespan(FrameDurationTicks)));
 	}
 }
 
diff --git a/Engine/Plugins/Media/MediaCompositing/Source/MediaCompositing/Private/MovieScene/MovieSceneMediaTemplate.h b/Engine/Plugins/Media/MediaCompositing/Source/MediaCompositing/Private/MovieScene/MovieSceneMediaTemplate.h
--- a/Engine/Plugins/Media/MediaCompositing/Source/MediaCompositing/Private/MovieScene/MovieSceneMediaTemplate.h
+++ b/Engine/Plugins/Media/MediaCompositing/Source/MediaCompositing/Private/MovieScene/MovieSceneMediaTemplate.h
@@ -2,6 +2,7 @@
 
 #pragma once
 
+#include "Channels/MovieSceneFloatChannel.h"
 #include "Evaluation/MovieSceneEvalTemplate.h"
 #include "Evaluation/MovieScenePropertyTemplate.h"
 #include "MovieSceneObjectBindingID.h"
@@ -51,6 +52,9 @@ struct FMovieSceneMediaSectionParams
 	UPROPERTY()
 	FFrameNumber StartFrameOffset;
 
+	UPROPERTY()
+	FMovieSceneFloatChannel ProxyTextureBlend;
+
 	UPROPERTY()
 	int32 ProxyTextureIndex;
 
diff --git a/Engine/Plugins/Media/MediaCompositing/Source/MediaCompositing/Public/MovieSceneMediaData.h b/Engine/Plugins/Media/MediaCompositing/Source/MediaCompositing/Public/MovieSceneMediaData.h
--- a/Engine/Plugins/Media/MediaCompositing/Source/MediaCompositing/Public/MovieSceneMediaData.h
+++ b/Engine/Plugins/Media/MediaCompositing/Source/MediaCompositing/Public/MovieSceneMediaData.h
@@ -38,6 +38,11 @@ public:
 	 */
 	UObject* GetPlayerProxy() { return PlayerProxy.Get(); }
 
+	/**
+	 * Get the texture index we are using (when using a proxy).
+	 */
+	int32 GetProxyTextureIndex() { return ProxyTextureIndex; }
+
 	/**
 	 * Set the time to seek to after opening a media source has finished.
 	 *
@@ -46,7 +51,7 @@ public:
 	void SeekOnOpen(FTimespan Time);
 
 	/** Set up this persistent data object. */
-	void Setup(UMediaPlayer* OverrideMediaPlayer, UObject* InPlayerProxy, int32 ProxyTextureIndex);
+	void Setup(UMediaPlayer* OverrideMediaPlayer, UObject* InPlayerProxy, int32 InProxyTextureIndex);
 
 	/**
 	 * Called from FMovieSceneMediaSectionTemplate::Initialize.
@@ -81,6 +86,8 @@ private:
 	TWeakObjectPtr<UObject> PlayerProxy;
 	/** Media texture allocated from the proxy. */
 	TWeakObjectPtr<UMediaTexture> ProxyMediaTexture;
+	/** Index of texture allocated from the proxy. */
+	int32 ProxyTextureIndex;
 
 	/** The time to seek to after the media source is opened. */
 	FTimespan SeekOnOpenTime;
diff --git a/Engine/Plugins/Media/MediaCompositing/Source/MediaCompositing/Public/MovieSceneMediaSection.h b/Engine/Plugins/Media/MediaCompositing/Source/MediaCompositing/Public/MovieSceneMediaSection.h
--- a/Engine/Plugins/Media/MediaCompositing/Source/MediaCompositing/Public/MovieSceneMediaSection.h
+++ b/Engine/Plugins/Media/MediaCompositing/Source/MediaCompositing/Public/MovieSceneMediaSection.h
@@ -4,8 +4,10 @@
 
 #include "MovieSceneSection.h"
 
+#include "Channels/MovieSceneChannelProxy.h"
 #include "MediaSource.h"
 #include "MovieSceneObjectBindingID.h"
+#include "Sections/MovieSceneFloatSection.h"
 
 #include "MovieSceneMediaSection.generated.h"
 
@@ -73,6 +75,10 @@ public:
 	UPROPERTY(Transient)
 	bool bHasMediaPlayerProxy;
 
+	/** If using an object like a MediaPlate, then this determines the blend value for our texture. */
+	UPROPERTY()
+	FMovieSceneFloatChannel TextureBlend;
+
 public:
 
 	/**
