MediaCompositing: Uses the section easing (instead of keys) to handle cross fading.

diff --git a/Engine/Plugins/Media/MediaCompositing/Source/MediaCompositing/Private/MovieScene/MovieSceneMediaSection.cpp b/Engine/Plugins/Media/MediaCompositing/Source/MediaCompositing/Private/MovieScene/MovieSceneMediaSection.cpp
--- a/Engine/Plugins/Media/MediaCompositing/Source/MediaCompositing/Private/MovieScene/MovieSceneMediaSection.cpp
+++ b/Engine/Plugins/Media/MediaCompositing/Source/MediaCompositing/Private/MovieScene/MovieSceneMediaSection.cpp
@@ -35,17 +35,7 @@ UMovieSceneMediaSection::UMovieSceneMediaSection(const FObjectInitializer& Objec
 #endif
 
 	EvalOptions.CompletionMode = EMovieSceneCompletionMode::RestoreState;
-	TextureBlend.SetDefault(1.0f);
-
-#if WITH_EDITOR
-
-	ChannelProxy = MakeShared<FMovieSceneChannelProxy>(TextureBlend, FMovieSceneChannelMetaData(), TMovieSceneExternalValue<float>::Make());
-
-#else
-
-	ChannelProxy = MakeShared<FMovieSceneChannelProxy>(TextureBlend);
-
-#endif
+	BlendType = EMovieSceneBlendType::Absolute;
 }
 
 void UMovieSceneMediaSection::PostInitProperties()
diff --git a/Engine/Plugins/Media/MediaCompositing/Source/MediaCompositing/Private/MovieScene/MovieSceneMediaTemplate.cpp b/Engine/Plugins/Media/MediaCompositing/Source/MediaCompositing/Private/MovieScene/MovieSceneMediaTemplate.cpp
--- a/Engine/Plugins/Media/MediaCompositing/Source/MediaCompositing/Private/MovieScene/MovieSceneMediaTemplate.cpp
+++ b/Engine/Plugins/Media/MediaCompositing/Source/MediaCompositing/Private/MovieScene/MovieSceneMediaTemplate.cpp
@@ -275,6 +275,7 @@ private:
  *****************************************************************************/
 
 FMovieSceneMediaSectionTemplate::FMovieSceneMediaSectionTemplate(const UMovieSceneMediaSection& InSection, const UMovieSceneMediaTrack& InTrack)
+	: MediaSection(&InSection)
 {
 	Params.MediaSource = InSection.GetMediaSource();
 	Params.MediaSourceProxy = InSection.GetMediaSourceProxy();
@@ -300,7 +301,6 @@ FMovieSceneMediaSectionTemplate::FMovieSceneMediaSectionTemplate(const UMovieSce
 		Params.SectionEndFrame = InSection.GetRange().GetUpperBoundValue();
 	}
 
-	Params.ProxyTextureBlend = InSection.TextureBlend;
 	Params.ProxyTextureIndex = InSection.TextureIndex;
 }
 
@@ -340,8 +340,7 @@ void FMovieSceneMediaSectionTemplate::Evaluate(const FMovieSceneEvaluationOperan
 		const double FrameDurationInSeconds = FMath::Max(FrameRate.AsSeconds(FFrameTime(1)), (Context.GetRange().Size<FFrameTime>()) / Context.GetFrameRate());
 		const int64 FrameDurationTicks = FrameDurationInSeconds * ETimespan::TicksPerSecond;
 
-		float ProxyTextureBlend = 1.0f;
-		Params.ProxyTextureBlend.Evaluate(Context.GetTime(), ProxyTextureBlend);
+		float ProxyTextureBlend = MediaSection->EvaluateEasing(Context.GetTime());
 
 		#if MOVIESCENEMEDIATEMPLATE_TRACE_EVALUATION
 			GLog->Logf(ELogVerbosity::Log, TEXT("Evaluating frame %i+%f, FrameRate %i/%i, FrameTicks %d, FrameDurationTicks %d"),
diff --git a/Engine/Plugins/Media/MediaCompositing/Source/MediaCompositing/Private/MovieScene/MovieSceneMediaTemplate.h b/Engine/Plugins/Media/MediaCompositing/Source/MediaCompositing/Private/MovieScene/MovieSceneMediaTemplate.h
--- a/Engine/Plugins/Media/MediaCompositing/Source/MediaCompositing/Private/MovieScene/MovieSceneMediaTemplate.h
+++ b/Engine/Plugins/Media/MediaCompositing/Source/MediaCompositing/Private/MovieScene/MovieSceneMediaTemplate.h
@@ -101,4 +101,7 @@ private:
 
 	UPROPERTY()
 	FMovieSceneMediaSectionParams Params;
+
+	UPROPERTY()
+	TObjectPtr<const UMovieSceneMediaSection> MediaSection;
 };
diff --git a/Engine/Plugins/Media/MediaCompositing/Source/MediaCompositing/Private/MovieScene/MovieSceneMediaTrack.cpp b/Engine/Plugins/Media/MediaCompositing/Source/MediaCompositing/Private/MovieScene/MovieSceneMediaTrack.cpp
--- a/Engine/Plugins/Media/MediaCompositing/Source/MediaCompositing/Private/MovieScene/MovieSceneMediaTrack.cpp
+++ b/Engine/Plugins/Media/MediaCompositing/Source/MediaCompositing/Private/MovieScene/MovieSceneMediaTrack.cpp
@@ -27,6 +27,8 @@ UMovieSceneMediaTrack::UMovieSceneMediaTrack(const FObjectInitializer& ObjectIni
 	EvalOptions.bEvaluateInPreroll = true;
 	EvalOptions.bEvaluateInPostroll = true;
 
+	SupportedBlendTypes.Add(EMovieSceneBlendType::Absolute);
+
 #if WITH_EDITORONLY_DATA
 	TrackTint = FColor(0, 0, 0, 200);
 	bSupportsDefaultSections = false;
diff --git a/Engine/Plugins/Media/MediaCompositing/Source/MediaCompositing/Public/MovieSceneMediaSection.h b/Engine/Plugins/Media/MediaCompositing/Source/MediaCompositing/Public/MovieSceneMediaSection.h
--- a/Engine/Plugins/Media/MediaCompositing/Source/MediaCompositing/Public/MovieSceneMediaSection.h
+++ b/Engine/Plugins/Media/MediaCompositing/Source/MediaCompositing/Public/MovieSceneMediaSection.h
@@ -4,10 +4,8 @@
 
 #include "MovieSceneSection.h"
 
-#include "Channels/MovieSceneChannelProxy.h"
 #include "MediaSource.h"
 #include "MovieSceneObjectBindingID.h"
-#include "Sections/MovieSceneFloatSection.h"
 
 #include "MovieSceneMediaSection.generated.h"
 
@@ -75,10 +73,6 @@ public:
 	UPROPERTY(Transient)
 	bool bHasMediaPlayerProxy;
 
-	/** If using an object like a MediaPlate, then this determines the blend value for our texture. */
-	UPROPERTY()
-	FMovieSceneFloatChannel TextureBlend;
-
 public:
 
 	/**
