MediaCompositing: Calculates the desired texture indices for cross fading sections.

diff --git a/Engine/Plugins/Media/MediaCompositing/Source/MediaCompositing/Private/MovieScene/MovieSceneMediaTemplate.cpp b/Engine/Plugins/Media/MediaCompositing/Source/MediaCompositing/Private/MovieScene/MovieSceneMediaTemplate.cpp
--- a/Engine/Plugins/Media/MediaCompositing/Source/MediaCompositing/Private/MovieScene/MovieSceneMediaTemplate.cpp
+++ b/Engine/Plugins/Media/MediaCompositing/Source/MediaCompositing/Private/MovieScene/MovieSceneMediaTemplate.cpp
@@ -300,8 +300,6 @@ FMovieSceneMediaSectionTemplate::FMovieSceneMediaSectionTemplate(const UMovieSce
 	{
 		Params.SectionEndFrame = InSection.GetRange().GetUpperBoundValue();
 	}
-
-	Params.ProxyTextureIndex = InSection.TextureIndex;
 }
 
 
@@ -369,6 +367,8 @@ void FMovieSceneMediaSectionTemplate::Initialize(const FMovieSceneEvaluationOper
 	FMovieSceneMediaData* SectionData = PersistentData.FindSectionData<FMovieSceneMediaData>();
 	if (SectionData == nullptr)
 	{
+		int32 ProxyTextureIndex = (MediaSection != nullptr) ? MediaSection->TextureIndex : 0;
+		
 		// Are we overriding the media player?
 		UMediaPlayer* MediaPlayer = Params.MediaPlayer;
 		UObject* PlayerProxy = nullptr;
@@ -387,6 +387,18 @@ void FMovieSceneMediaSectionTemplate::Initialize(const FMovieSceneEvaluationOper
 						if (BoundObject != nullptr)
 						{
 							MediaAssetsModule->GetPlayerFromObject(BoundObject, PlayerProxy);
+							if (PlayerProxy != nullptr)
+							{
+								IMediaPlayerProxyInterface* PlayerProxyInterface = Cast<IMediaPlayerProxyInterface>(PlayerProxy);
+								if (PlayerProxyInterface != nullptr)
+								{
+									int32 NumberOfTextures = PlayerProxyInterface->ProxyGetNumberOfMediaTextures();
+									if (NumberOfTextures > 0)
+									{
+										ProxyTextureIndex %= NumberOfTextures;
+									}
+								}
+							}
 							break;
 						}
 					}
@@ -396,7 +408,7 @@ void FMovieSceneMediaSectionTemplate::Initialize(const FMovieSceneEvaluationOper
 
 		// Add section data.
 		SectionData = &PersistentData.AddSectionData<FMovieSceneMediaData>();
-		SectionData->Setup(MediaPlayer, PlayerProxy, Params.ProxyTextureIndex);
+		SectionData->Setup(MediaPlayer, PlayerProxy, ProxyTextureIndex);
 	}
 
 	if (!ensure(SectionData != nullptr))
@@ -490,6 +502,16 @@ void FMovieSceneMediaSectionTemplate::TearDown(FPersistentEvaluationData& Persis
 		Params.MediaTexture->SetMediaPlayer(nullptr);
 	}
 
+	UObject* PlayerProxy = SectionData->GetPlayerProxy();
+	if (PlayerProxy != nullptr)
+	{
+		IMediaPlayerProxyInterface* PlayerProxyInterface = Cast<IMediaPlayerProxyInterface>(PlayerProxy);
+		if (PlayerProxyInterface != nullptr)
+		{
+			PlayerProxyInterface->ProxySetTextureBlend(SectionData->GetProxyTextureIndex(), 0.0f);
+		}
+	}
+
 	SectionData->TearDown();
 }
 
diff --git a/Engine/Plugins/Media/MediaCompositing/Source/MediaCompositing/Private/MovieScene/MovieSceneMediaTemplate.h b/Engine/Plugins/Media/MediaCompositing/Source/MediaCompositing/Private/MovieScene/MovieSceneMediaTemplate.h
--- a/Engine/Plugins/Media/MediaCompositing/Source/MediaCompositing/Private/MovieScene/MovieSceneMediaTemplate.h
+++ b/Engine/Plugins/Media/MediaCompositing/Source/MediaCompositing/Private/MovieScene/MovieSceneMediaTemplate.h
@@ -55,9 +55,6 @@ struct FMovieSceneMediaSectionParams
 	UPROPERTY()
 	FMovieSceneFloatChannel ProxyTextureBlend;
 
-	UPROPERTY()
-	int32 ProxyTextureIndex;
-
 	FMovieSceneMediaSectionParams()
 		: MediaSoundComponent(nullptr)
 		, MediaSource(nullptr)
@@ -65,7 +62,6 @@ struct FMovieSceneMediaSectionParams
 		, MediaTexture(nullptr)
 		, MediaPlayer(nullptr)
 		, bLooping(false)
-		, ProxyTextureIndex(0)
 	{}
 };
 
diff --git a/Engine/Plugins/Media/MediaCompositing/Source/MediaCompositing/Private/MovieScene/MovieSceneMediaTrack.cpp b/Engine/Plugins/Media/MediaCompositing/Source/MediaCompositing/Private/MovieScene/MovieSceneMediaTrack.cpp
--- a/Engine/Plugins/Media/MediaCompositing/Source/MediaCompositing/Private/MovieScene/MovieSceneMediaTrack.cpp
+++ b/Engine/Plugins/Media/MediaCompositing/Source/MediaCompositing/Private/MovieScene/MovieSceneMediaTrack.cpp
@@ -67,6 +67,7 @@ UMovieSceneMediaSection* UMovieSceneMediaTrack::AddNewSectionOnRow(FFrameNumber
 	NewSection->InitialPlacementOnRow(MediaSections, Time, DurationToUse.FrameNumber.Value, RowIndex);
 
 	MediaSections.Add(NewSection);
+	UpdateSectionTextureIndices();
 
 	return NewSection;
 }
@@ -78,6 +79,7 @@ UMovieSceneMediaSection* UMovieSceneMediaTrack::AddNewSectionOnRow(FFrameNumber
 void UMovieSceneMediaTrack::AddSection(UMovieSceneSection& Section)
 {
 	MediaSections.Add(&Section);
+	UpdateSectionTextureIndices();
 }
 
 
@@ -114,18 +116,116 @@ bool UMovieSceneMediaTrack::IsEmpty() const
 void UMovieSceneMediaTrack::RemoveSection(UMovieSceneSection& Section)
 {
 	MediaSections.Remove(&Section);
+	UpdateSectionTextureIndices();
 }
 
 void UMovieSceneMediaTrack::RemoveSectionAt(int32 SectionIndex)
 {
 	MediaSections.RemoveAt(SectionIndex);
+	UpdateSectionTextureIndices();
 }
 
+#if WITH_EDITOR
+
+EMovieSceneSectionMovedResult UMovieSceneMediaTrack::OnSectionMoved(UMovieSceneSection& InSection, const FMovieSceneSectionMovedParams& Params)
+{
+	UpdateSectionTextureIndices();
+	return EMovieSceneSectionMovedResult::None;
+}
+
+#endif // WITH_EDITOR
 
 FMovieSceneEvalTemplatePtr UMovieSceneMediaTrack::CreateTemplateForSection(const UMovieSceneSection& InSection) const
 {
 	return FMovieSceneMediaSectionTemplate(*CastChecked<const UMovieSceneMediaSection>(&InSection), *this);
 }
 
+void UMovieSceneMediaTrack::UpdateSectionTextureIndices()
+{
+#if WITH_EDITOR
+	// Do we have an object binding?
+	UMovieScene* MovieScene = Cast<UMovieScene>(GetOuter());
+	if (MovieScene != nullptr)
+	{
+		const TArray<FMovieSceneBinding>& Bindings = MovieScene->GetBindings();
+		for (const FMovieSceneBinding& Binding : Bindings)
+		{
+			bool bIsThisMyBinding = false;
+			const TArray<UMovieSceneTrack*>& Tracks = Binding.GetTracks();
+			for (UMovieSceneTrack* Track : Tracks)
+			{
+				if ((Track != nullptr) && (Track == this))
+				{
+					bIsThisMyBinding = true;
+					break;
+				}
+			}
+
+			if (bIsThisMyBinding)
+			{
+				// Get all sections.
+				TArray<UMovieSceneMediaSection*> AllSections;
+				for (UMovieSceneTrack* Track : Tracks)
+				{
+					if (Track != nullptr)
+					{
+						const TArray<UMovieSceneSection*>& Sections = Track->GetAllSections();
+						for (UMovieSceneSection* Section : Sections)
+						{
+							UMovieSceneMediaSection* MediaSection = Cast<UMovieSceneMediaSection>(Section);
+							if (MediaSection != nullptr)
+							{
+								AllSections.Add(MediaSection);
+							}
+						}
+					}
+				}
+
+				// Set up indices from earliest section to latest.
+				AllSections.Sort([](UMovieSceneMediaSection& A, UMovieSceneMediaSection& B)
+				{ return A.GetRange().GetLowerBoundValue() < B.GetRange().GetLowerBoundValue(); });
+
+				int32 CurrentTextureIndex = 0;
+				UMovieSceneMediaSection* PreviousSection = nullptr;
+				TRange<FFrameNumber> PreviousRange;
+				for (UMovieSceneMediaSection* Section : AllSections)
+				{
+					TRange <FFrameNumber> Range = Section->GetRange();
+					
+					// If we overlap the previous section then we need another texture so increment
+					// the count. Otherwise we can reuse the same index as the previous section.
+					bool bIsOverlappingPreviousSection = false;
+					if (PreviousRange.HasUpperBound())
+					{
+						bIsOverlappingPreviousSection = Range.GetLowerBoundValue() <
+							PreviousRange.GetUpperBoundValue();
+					}
+					if (bIsOverlappingPreviousSection)
+					{
+						CurrentTextureIndex = (CurrentTextureIndex + 1);
+					}
+					else if (PreviousSection != nullptr)
+					{
+						CurrentTextureIndex = PreviousSection->TextureIndex;
+					}
+
+					Section->TextureIndex = CurrentTextureIndex;
+
+					// Previous section is defined as the section right before in the timeline, which
+					// is not necessarily the last section we looked at.
+					if ((PreviousRange.HasUpperBound() == false) ||
+						(PreviousRange.GetUpperBoundValue() < Range.GetUpperBoundValue()))
+					{
+						PreviousSection = Section;
+						PreviousRange = Range;
+					}
+				}
+				break;
+			}
+		}
+	}
+#endif // WITH_EDITOR
+}
+
 #undef LOCTEXT_NAMESPACE
 
diff --git a/Engine/Plugins/Media/MediaCompositing/Source/MediaCompositing/Public/MovieSceneMediaTrack.h b/Engine/Plugins/Media/MediaCompositing/Source/MediaCompositing/Public/MovieSceneMediaTrack.h
--- a/Engine/Plugins/Media/MediaCompositing/Source/MediaCompositing/Public/MovieSceneMediaTrack.h
+++ b/Engine/Plugins/Media/MediaCompositing/Source/MediaCompositing/Public/MovieSceneMediaTrack.h
@@ -63,6 +63,13 @@ public:
 	virtual void RemoveSectionAt(int32 SectionIndex) override;
 	virtual bool SupportsMultipleRows() const override { return true; }
 
+
+#if WITH_EDITOR
+	
+	virtual EMovieSceneSectionMovedResult OnSectionMoved(UMovieSceneSection& Section, const FMovieSceneSectionMovedParams& Params);
+
+#endif // WITH_EDITOR
+
 private:
 	/** Base function to add a new section. */
 	UMovieSceneMediaSection* AddNewSectionOnRow(FFrameNumber Time, int32 RowIndex);
@@ -70,4 +77,11 @@ private:
 	/** List of all media sections. */
 	UPROPERTY()
 	TArray<TObjectPtr<UMovieSceneSection>> MediaSections;
+
+	/**
+	 * Updates all our sections so they have correct texture indices
+	 * so the proxy can blend all the sections together correctly.
+	 */
+	void UpdateSectionTextureIndices();
+
 };
diff --git a/Engine/Plugins/Media/MediaPlate/Source/MediaPlate/Private/MediaPlateComponent.cpp b/Engine/Plugins/Media/MediaPlate/Source/MediaPlate/Private/MediaPlateComponent.cpp
--- a/Engine/Plugins/Media/MediaPlate/Source/MediaPlate/Private/MediaPlateComponent.cpp
+++ b/Engine/Plugins/Media/MediaPlate/Source/MediaPlate/Private/MediaPlateComponent.cpp
@@ -637,6 +637,11 @@ UMediaTexture* UMediaPlateComponent::ProxyGetMediaTexture(int32 TextureIndex)
 	return GetMediaTexture(TextureIndex);
 }
 
+int32 UMediaPlateComponent::ProxyGetNumberOfMediaTextures() const
+{
+	return MediaTextures.Num();
+}
+
 bool UMediaPlateComponent::ProxyIsPlaylistIndexPlaying(int32 Index) const
 {
 	return bIsMediaPlatePlaying && (Index == PlaylistIndex);
diff --git a/Engine/Plugins/Media/MediaPlate/Source/MediaPlate/Public/MediaPlateComponent.h b/Engine/Plugins/Media/MediaPlate/Source/MediaPlate/Public/MediaPlateComponent.h
--- a/Engine/Plugins/Media/MediaPlate/Source/MediaPlate/Public/MediaPlateComponent.h
+++ b/Engine/Plugins/Media/MediaPlate/Source/MediaPlate/Public/MediaPlateComponent.h
@@ -254,6 +254,7 @@ public:
 	virtual UMediaSource* GetMediaSourceFromIndex(int32 Index) const override;
 	virtual void ProxyClose() override;
 	virtual UMediaTexture* ProxyGetMediaTexture(int32 TextureIndex) override;
+	virtual int32 ProxyGetNumberOfMediaTextures() const override;
 	virtual bool ProxyIsPlaylistIndexPlaying(int32 Index) const override;
 	virtual void ProxyOpenPlaylistIndex(int32 Index) override;
 	virtual void ProxySetPlayOnOpen(bool bInPlayOnOpen) override;
diff --git a/Engine/Source/Runtime/MediaAssets/Public/MediaPlayerProxyInterface.h b/Engine/Source/Runtime/MediaAssets/Public/MediaPlayerProxyInterface.h
--- a/Engine/Source/Runtime/MediaAssets/Public/MediaPlayerProxyInterface.h
+++ b/Engine/Source/Runtime/MediaAssets/Public/MediaPlayerProxyInterface.h
@@ -68,6 +68,11 @@ public:
 	 */
 	virtual UMediaTexture* ProxyGetMediaTexture(int32 TextureIndex) = 0;
 
+	/**
+	 * Gets the number of media textures we can get with ProxyGetMediaTexture.
+	 */
+	virtual int32 ProxyGetNumberOfMediaTextures() const = 0;
+
 	/**
 	 * Ask if a specific track in the playlist is playing.
 	 */
