MediaCompositing: Set looping when receiving the media open event in case there was no player when we set the loop previously.

diff --git a/Engine/Plugins/Media/MediaCompositing/Source/MediaCompositing/Private/MovieScene/MovieSceneMediaData.cpp b/Engine/Plugins/Media/MediaCompositing/Source/MediaCompositing/Private/MovieScene/MovieSceneMediaData.cpp
--- a/Engine/Plugins/Media/MediaCompositing/Source/MediaCompositing/Private/MovieScene/MovieSceneMediaData.cpp
+++ b/Engine/Plugins/Media/MediaCompositing/Source/MediaCompositing/Private/MovieScene/MovieSceneMediaData.cpp
@@ -147,6 +147,10 @@ void FMovieSceneMediaData::HandleMediaPlayerEvent(EMediaEvent Event)
 	    return;
     }
 
+	// Update looping here, as if there is no IMediaPlayer (which is different to MediaPlayer)
+	// when we last called SetLooping then looping might not be set correctly.
+	MediaPlayer->SetLooping(MediaPlayer->IsLooping());
+
 	FTimespan MediaTime;
 	if (!MediaPlayer->IsLooping())
 	{
