MediaCompositing: Fixed textures not being reused.
#rnx

diff --git a/Engine/Plugins/Media/MediaCompositing/Source/MediaCompositing/Private/MovieScene/MovieSceneMediaTrack.cpp b/Engine/Plugins/Media/MediaCompositing/Source/MediaCompositing/Private/MovieScene/MovieSceneMediaTrack.cpp
--- a/Engine/Plugins/Media/MediaCompositing/Source/MediaCompositing/Private/MovieScene/MovieSceneMediaTrack.cpp
+++ b/Engine/Plugins/Media/MediaCompositing/Source/MediaCompositing/Private/MovieScene/MovieSceneMediaTrack.cpp
@@ -186,38 +186,46 @@ void UMovieSceneMediaTrack::UpdateSectionTextureIndices()
 				{ return A.GetRange().GetLowerBoundValue() < B.GetRange().GetLowerBoundValue(); });
 
 				int32 CurrentTextureIndex = 0;
-				UMovieSceneMediaSection* PreviousSection = nullptr;
-				TRange<FFrameNumber> PreviousRange;
+				TSet<int32> AllocatedTextures;
+				TArray<UMovieSceneMediaSection*>ActiveSections;
+
 				for (UMovieSceneMediaSection* Section : AllSections)
 				{
 					TRange <FFrameNumber> Range = Section->GetRange();
-					
-					// If we overlap the previous section then we need another texture so increment
-					// the count. Otherwise we can reuse the same index as the previous section.
-					bool bIsOverlappingPreviousSection = false;
-					if (PreviousRange.HasUpperBound())
-					{
-						bIsOverlappingPreviousSection = Range.GetLowerBoundValue() <
-							PreviousRange.GetUpperBoundValue();
-					}
-					if (bIsOverlappingPreviousSection)
-					{
-						CurrentTextureIndex = (CurrentTextureIndex + 1);
-					}
-					else if (PreviousSection != nullptr)
+					FFrameNumber RangeStart = Range.GetLowerBoundValue();
+
+					// Remove expired sections.
+					for (int32 Index = 0; Index < ActiveSections.Num();)
 					{
-						CurrentTextureIndex = PreviousSection->TextureIndex;
+						UMovieSceneMediaSection* ActiveSection = ActiveSections[Index];
+						if (ActiveSection->GetRange().GetUpperBoundValue() <= RangeStart)
+						{
+							if (CurrentTextureIndex > ActiveSection->TextureIndex)
+							{
+								CurrentTextureIndex = ActiveSection->TextureIndex;
+							}
+							AllocatedTextures.Remove(ActiveSection->TextureIndex);
+							ActiveSections.RemoveAtSwap(Index);
+						}
+						else
+						{
+							Index++;
+						}
 					}
-
+					
+					// Set up our section.
 					Section->TextureIndex = CurrentTextureIndex;
+					AllocatedTextures.Add(CurrentTextureIndex);
+					ActiveSections.Add(Section);
 
-					// Previous section is defined as the section right before in the timeline, which
-					// is not necessarily the last section we looked at.
-					if ((PreviousRange.HasUpperBound() == false) ||
-						(PreviousRange.GetUpperBoundValue() < Range.GetUpperBoundValue()))
+					// Get the next texture index to use.
+					while (true)
 					{
-						PreviousSection = Section;
-						PreviousRange = Range;
+						CurrentTextureIndex++;
+						if (AllocatedTextures.Contains(CurrentTextureIndex) == false)
+						{
+							break;
+						}
 					}
 				}
 				break;
