MediaCompositing: Fixed sections not getting the appropriate length when not adding a media plate.
#rnx

diff --git a/Engine/Plugins/Media/MediaCompositing/Source/MediaCompositing/Private/MovieScene/MovieSceneMediaTrack.cpp b/Engine/Plugins/Media/MediaCompositing/Source/MediaCompositing/Private/MovieScene/MovieSceneMediaTrack.cpp
--- a/Engine/Plugins/Media/MediaCompositing/Source/MediaCompositing/Private/MovieScene/MovieSceneMediaTrack.cpp
+++ b/Engine/Plugins/Media/MediaCompositing/Source/MediaCompositing/Private/MovieScene/MovieSceneMediaTrack.cpp
@@ -2,8 +2,12 @@
 
 #include "MovieSceneMediaTrack.h"
 
+#include "IMediaClock.h"
+#include "IMediaClockSink.h"
+#include "IMediaModule.h"
 #include "MediaPlayer.h"
 #include "MediaSource.h"
+#include "Modules/ModuleManager.h"
 #include "MovieScene.h"
 #include "MovieSceneMediaSection.h"
 #include "MovieSceneMediaTemplate.h"
@@ -15,12 +19,53 @@
 
 #define LOCTEXT_NAMESPACE "MovieSceneMediaTrack"
 
+#if WITH_EDITOR
+
+/**
+ * Media clock sink for media track.
+ */
+class FMovieSceneMediaTrackClockSink
+	: public IMediaClockSink
+{
+public:
+
+	FMovieSceneMediaTrackClockSink(UMovieSceneMediaTrack* InOwner)
+		: Owner(InOwner)
+	{ }
+
+	virtual ~FMovieSceneMediaTrackClockSink() { }
+
+	virtual void TickOutput(FTimespan DeltaTime, FTimespan Timecode) override
+	{
+		if (UMovieSceneMediaTrack* OwnerPtr = Owner.Get())
+		{
+			Owner->TickOutput();
+		}
+	}
+
+	/**
+	 * Call this when the owner is destroyed.
+	 */
+	void OwnerDestroyed()
+	{
+		Owner.Reset();
+	}
+
+private:
+
+	TWeakObjectPtr<UMovieSceneMediaTrack> Owner;
+};
+
+#endif // WITH_EDITOR
 
 /* UMovieSceneMediaTrack interface
  *****************************************************************************/
 
 UMovieSceneMediaTrack::UMovieSceneMediaTrack(const FObjectInitializer& ObjectInitializer)
 	: Super(ObjectInitializer)
+#if WITH_EDITORONLY_DATA
+	, bGetDurationDelay(false)
+#endif // WITH_EDITORONLY_DATA
 {
 	EvalOptions.bCanEvaluateNearestSection = false;
 	EvalOptions.bEvalNearestSection = false;
@@ -43,14 +88,16 @@ UMovieSceneSection* UMovieSceneMediaTrack::AddNewMediaSourceOnRow(UMediaSource&
 {
 	UMovieSceneMediaSection* NewSection = AddNewSectionOnRow(Time, RowIndex);
 	NewSection->SetMediaSource(&MediaSource);
+	StartGetDuration(&MediaSource, NewSection);
 	return NewSection;
 }
 
 
-UMovieSceneSection* UMovieSceneMediaTrack::AddNewMediaSourceProxyOnRow(const FMovieSceneObjectBindingID& ObjectBinding, int32 MediaSourceProxyIndex, FFrameNumber Time, int32 RowIndex)
+UMovieSceneSection* UMovieSceneMediaTrack::AddNewMediaSourceProxyOnRow(UMediaSource* MediaSource, const FMovieSceneObjectBindingID& ObjectBinding, int32 MediaSourceProxyIndex, FFrameNumber Time, int32 RowIndex)
 {
 	UMovieSceneMediaSection* NewSection = AddNewSectionOnRow(Time, RowIndex);
 	NewSection->SetMediaSourceProxy(ObjectBinding, MediaSourceProxyIndex);
+	StartGetDuration(MediaSource, NewSection);
 	return NewSection;
 }
 
@@ -72,6 +119,56 @@ UMovieSceneMediaSection* UMovieSceneMediaTrack::AddNewSectionOnRow(FFrameNumber
 	return NewSection;
 }
 
+void UMovieSceneMediaTrack::TickOutput()
+{
+#if WITH_EDITOR
+	// Do we have any new sections that need durations?
+	if (NewSections.Num() > 0)
+	{
+		// Can we get the duration?
+		if (bGetDurationDelay)
+		{
+			bGetDurationDelay = false;
+		}
+		else
+		{
+			// Loop over all new sections.
+			for (int32 Index = 0; Index < NewSections.Num();)
+			{
+				if (NewSections[Index].Key.IsValid())
+				{
+					// Try and get the duration.
+					if (GetDuration(NewSections[Index].Key, NewSections[Index].Value))
+					{
+						NewSections.RemoveAtSwap(Index);
+					}
+					else
+					{
+						++Index;
+					}
+				}
+				else
+				{
+					NewSections.RemoveAtSwap(Index);
+				}
+			}
+		}
+	}
+
+	if (NewSections.Num() == 0)
+	{
+		if (ClockSink.IsValid())
+		{
+			IMediaModule* MediaModule = FModuleManager::GetModulePtr<IMediaModule>("Media");
+			if (MediaModule != nullptr)
+			{
+				MediaModule->GetClock().RemoveSink(ClockSink.ToSharedRef());
+			}
+		}
+	}
+#endif // WITH_EDITOR
+}
+
 
 /* UMovieScenePropertyTrack interface
  *****************************************************************************/
@@ -127,6 +224,25 @@ void UMovieSceneMediaTrack::RemoveSectionAt(int32 SectionIndex)
 
 #if WITH_EDITOR
 
+void UMovieSceneMediaTrack::BeginDestroy()
+{
+	if (ClockSink.IsValid())
+	{
+		// Tell sink we are done.
+		ClockSink->OwnerDestroyed();
+
+		IMediaModule* MediaModule = FModuleManager::GetModulePtr<IMediaModule>("Media");
+		if (MediaModule != nullptr)
+		{
+			MediaModule->GetClock().RemoveSink(ClockSink.ToSharedRef());
+		}
+
+		ClockSink.Reset();
+	}
+
+	Super::BeginDestroy();
+}
+
 EMovieSceneSectionMovedResult UMovieSceneMediaTrack::OnSectionMoved(UMovieSceneSection& InSection, const FMovieSceneSectionMovedParams& Params)
 {
 	UpdateSectionTextureIndices();
@@ -235,5 +351,68 @@ void UMovieSceneMediaTrack::UpdateSectionTextureIndices()
 #endif // WITH_EDITOR
 }
 
+#if WITH_EDITOR
+
+void UMovieSceneMediaTrack::StartGetDuration(UMediaSource* MediaSource, UMovieSceneSection* Section)
+{
+	// Create media player.
+	TStrongObjectPtr<UMediaPlayer> MediaPlayer = TStrongObjectPtr<UMediaPlayer>(
+		NewObject<UMediaPlayer>(GetTransientPackage(),
+			MakeUniqueObjectName(GetTransientPackage(),
+				UMediaPlayer::StaticClass())));
+
+	// Open the media.
+	MediaPlayer->PlayOnOpen = false;
+	if (MediaPlayer->OpenSource(MediaSource))
+	{
+		NewSections.Emplace(MediaPlayer, Section);
+	}
+
+	// Some players like Electra report that they are closed at this point, so wait a frame.
+	bGetDurationDelay = true;
+
+	// Start the clock sink so we can tick.
+	IMediaModule* MediaModule = FModuleManager::LoadModulePtr<IMediaModule>("Media");
+	if (MediaModule != nullptr)
+	{
+		if (ClockSink.IsValid() == false)
+		{
+			ClockSink = MakeShared<FMovieSceneMediaTrackClockSink, ESPMode::ThreadSafe>(this);
+		}
+		MediaModule->GetClock().AddSink(ClockSink.ToSharedRef());
+	}
+}
+
+bool UMovieSceneMediaTrack::GetDuration(
+	const TStrongObjectPtr<UMediaPlayer>& MediaPlayer, TWeakObjectPtr<UMovieSceneSection>& NewSection)
+{
+	bool bIsDone = false;
+
+	// Check everything is ok.
+	if ((MediaPlayer.IsValid() == false) || (MediaPlayer->HasError()) || (MediaPlayer->IsClosed()) ||
+		(NewSection.IsValid() == false))
+	{
+		bIsDone = true;
+	}
+	else
+	{
+		// Get the duration.
+		FTimespan Duration = MediaPlayer->GetDuration();
+		if (Duration != 0)
+		{
+			// Once it is non zero, then set the length of the section.
+			FFrameRate TickResolution = NewSection->GetTypedOuter<UMovieScene>()->GetTickResolution();
+			FFrameNumber StartFrame = NewSection->GetInclusiveStartFrame();
+			FFrameNumber EndFrame = StartFrame + (Duration.GetTotalSeconds() * TickResolution).FrameNumber;
+			NewSection->SetEndFrame(TRangeBound<FFrameNumber>::Exclusive(EndFrame));
+			bIsDone = true;
+		}
+	}
+
+	return bIsDone;
+}
+
+#endif // WITH_EDITOR
+
 #undef LOCTEXT_NAMESPACE
 
diff --git a/Engine/Plugins/Media/MediaCompositing/Source/MediaCompositing/Public/MovieSceneMediaTrack.h b/Engine/Plugins/Media/MediaCompositing/Source/MediaCompositing/Public/MovieSceneMediaTrack.h
--- a/Engine/Plugins/Media/MediaCompositing/Source/MediaCompositing/Public/MovieSceneMediaTrack.h
+++ b/Engine/Plugins/Media/MediaCompositing/Source/MediaCompositing/Public/MovieSceneMediaTrack.h
@@ -5,10 +5,13 @@
 #include "Misc/Guid.h"
 #include "MovieSceneNameableTrack.h"
 #include "UObject/ObjectMacros.h"
+#include "UObject/StrongObjectPtr.h"
 #include "Compilation/IMovieSceneTrackTemplateProducer.h"
 
 #include "MovieSceneMediaTrack.generated.h"
 
+class FMovieSceneMediaTrackClockSink;
+class UMediaPlayer;
 class UMediaSource;
 class UMovieSceneMediaSection;
 struct FMovieSceneObjectBindingID;
@@ -38,16 +41,21 @@ public:
 	/** Adds a new media source to the track. */
 	virtual UMovieSceneSection* AddNewMediaSourceOnRow(UMediaSource& MediaSource, FFrameNumber Time, int32 RowIndex);
 	/** Adds a new media source to the track. */
-	virtual UMovieSceneSection* AddNewMediaSourceProxyOnRow(const FMovieSceneObjectBindingID& ObjectBinding, int32 MediaSourceProxyIndex, FFrameNumber Time, int32 RowIndex);
+	virtual UMovieSceneSection* AddNewMediaSourceProxyOnRow(UMediaSource* MediaSource, const FMovieSceneObjectBindingID& ObjectBinding, int32 MediaSourceProxyIndex, FFrameNumber Time, int32 RowIndex);
 
 	/** Adds a new media source on the next available/non-overlapping row. */
 	virtual UMovieSceneSection* AddNewMediaSource(UMediaSource& MediaSource, FFrameNumber Time) { return AddNewMediaSourceOnRow(MediaSource, Time, INDEX_NONE); }
 	/** Adds a new media source on the next available/non-overlapping row. */
-	virtual UMovieSceneSection* AddNewMediaSourceProxy(const FMovieSceneObjectBindingID& ObjectBinding, int32 MediaSourceProxyIndex, FFrameNumber Time)
+	virtual UMovieSceneSection* AddNewMediaSourceProxy(UMediaSource* MediaSource, const FMovieSceneObjectBindingID& ObjectBinding, int32 MediaSourceProxyIndex, FFrameNumber Time)
 	{
-		return AddNewMediaSourceProxyOnRow(ObjectBinding, MediaSourceProxyIndex, Time, INDEX_NONE);
+		return AddNewMediaSourceProxyOnRow(MediaSource, ObjectBinding, MediaSourceProxyIndex, Time, INDEX_NONE);
 	}
 
+	/**
+	 * Called from the media clock.
+	 */
+	void TickOutput();
+
 public:
 
 	//~ UMovieScenePropertyTrack interface
@@ -65,6 +73,8 @@ public:
 
 
 #if WITH_EDITOR
+
+	virtual void BeginDestroy() override;
 	
 	virtual EMovieSceneSectionMovedResult OnSectionMoved(UMovieSceneSection& Section, const FMovieSceneSectionMovedParams& Params);
 
@@ -78,10 +88,44 @@ private:
 	UPROPERTY()
 	TArray<TObjectPtr<UMovieSceneSection>> MediaSections;
 
+#if WITH_EDITORONLY_DATA
+
+	/** List of players that are we are trying to get durations from for the corresponding sections. */
+	TArray<TPair<TStrongObjectPtr<UMediaPlayer>, TWeakObjectPtr<UMovieSceneSection>>> NewSections;
+	/** Our media clock sink. */
+	TSharedPtr<FMovieSceneMediaTrackClockSink, ESPMode::ThreadSafe> ClockSink;
+	/** If true then don't check for the duration this frame. */
+	bool bGetDurationDelay;
+
+#endif // WITH_EDITORONLY_DATA
+
 	/**
 	 * Updates all our sections so they have correct texture indices
 	 * so the proxy can blend all the sections together correctly.
 	 */
 	void UpdateSectionTextureIndices();
 
+#if WITH_EDITOR
+	/**
+	 * Starts the process to get the duration of the media.
+	 * It might take a frame or more.
+	 *
+	 * @param MediaSource		Media to inspect.
+	 * @param Section			Will set this sequencer section to the length of the media.
+	*/
+	void StartGetDuration(UMediaSource* MediaSource, UMovieSceneSection* Section);
+
+	/**
+	 * Call this after StartGetDuration to try and get the duration of the media.
+	 *
+	 * @param MediaPlayer		Player that is opening the media.
+	 * @param NewSection		Movie section will be set to the media duration.
+	 * @return True if it is done and the player can be removed.
+	 */
+	bool GetDuration(const TStrongObjectPtr<UMediaPlayer>& MediaPlayer,
+		TWeakObjectPtr<UMovieSceneSection>& NewSection);
+#else  // WITH_EDITOR
+	void StartGetDuration(UMediaSource* MediaSource, UMovieSceneSection* Section) {}
+#endif // WITH_EDITOR
+
 };
diff --git a/Engine/Plugins/Media/MediaPlate/Source/MediaPlateEditor/Private/Sequencer/MediaPlateTrackEditor.cpp b/Engine/Plugins/Media/MediaPlate/Source/MediaPlateEditor/Private/Sequencer/MediaPlateTrackEditor.cpp
--- a/Engine/Plugins/Media/MediaPlate/Source/MediaPlateEditor/Private/Sequencer/MediaPlateTrackEditor.cpp
+++ b/Engine/Plugins/Media/MediaPlate/Source/MediaPlateEditor/Private/Sequencer/MediaPlateTrackEditor.cpp
@@ -8,7 +8,6 @@
 #include "MediaTexture.h"
 #include "MediaPlate.h"
 #include "MediaPlateComponent.h"
-#include "MediaPlayer.h"
 #include "MediaPlaylist.h"
 #include "Misc/MessageDialog.h"
 #include "Modules/ModuleManager.h"
@@ -33,7 +32,6 @@ TArray<FAnimatedPropertyKey, TInlineAllocator<1>> FMediaPlateTrackEditor::GetAni
 
 FMediaPlateTrackEditor::FMediaPlateTrackEditor(TSharedRef<ISequencer> InSequencer)
 	: FMovieSceneTrackEditor(InSequencer)
-	, bGetDurationDelay(false)
 {
 	OnActorAddedToSequencerHandle = InSequencer->OnActorAddedToSequencer().AddRaw(this, &FMediaPlateTrackEditor::HandleActorAdded);
 	FMediaTrackEditor::OnBuildOutlinerEditWidget.AddRaw(this,
@@ -84,42 +82,6 @@ bool FMediaPlateTrackEditor::SupportsType(TSubclassOf<UMovieSceneTrack> TrackCla
 	return false;
 }
 
-void FMediaPlateTrackEditor::Tick(float DeltaTime)
-{
-	// Do we have any new sections that need durations?
-	if (NewSections.Num() > 0)
-	{
-		// Can we get the duration?
-		if (bGetDurationDelay)
-		{
-			bGetDurationDelay = false;
-		}
-		else
-		{
-			// Loop over all new sections.
-			for (int32 Index = 0; Index < NewSections.Num();)
-			{
-				if (NewSections[Index].Key.IsValid())
-				{
-					// Try and get the duration.
-					if (GetDuration(NewSections[Index].Key, NewSections[Index].Value))
-					{
-						NewSections.RemoveAtSwap(Index);
-					}
-					else
-					{
-						++Index;
-					}
-				}
-				else
-				{
-					NewSections.RemoveAtSwap(Index);
-				}
-			}
-		}
-	}
-}
-
 void FMediaPlateTrackEditor::HandleAddMediaTrackToObjectBindingMenuEntryExecute(TArray<FGuid> InObjectBindingIDs)
 {
 	UMovieScene* FocusedMovieScene = GetFocusedMovieScene();
@@ -193,12 +155,9 @@ void FMediaPlateTrackEditor::AddTrackForComponent(UMediaPlateComponent* Componen
 				if (MediaSource != nullptr)
 				{
 					UMovieSceneSection* Section = MediaTrack->AddNewMediaSourceProxy(
-						TargetObjectBindingID, Index, FFrameNumber(0));
+						MediaSource, TargetObjectBindingID, Index, FFrameNumber(0));
 					if (Section != nullptr)
 					{
-						// Start process to get the duration.
-						StartGetDuration(MediaSource, Section);
-
 						// Tell the section it has a player proxy.
 						UMovieSceneMediaSection* MediaSection = Cast<UMovieSceneMediaSection>(Section);
 						if (MediaSection != nullptr)
@@ -261,54 +220,6 @@ void FMediaPlateTrackEditor::ImportObjectBinding(const TArray<FGuid> ObjectBindi
 	}
 }
 
-void FMediaPlateTrackEditor::StartGetDuration(UMediaSource* MediaSource, UMovieSceneSection* Section)
-{
-	// Create media player.
-	TStrongObjectPtr<UMediaPlayer> MediaPlayer = TStrongObjectPtr<UMediaPlayer>(
-		NewObject<UMediaPlayer>(GetTransientPackage(),
-			MakeUniqueObjectName(GetTransientPackage(),
-				UMediaPlayer::StaticClass())));
-
-	// Open the media.
-	MediaPlayer->PlayOnOpen = false;
-	if (MediaPlayer->OpenSource(MediaSource))
-	{
-		NewSections.Emplace(MediaPlayer, Section);
-	}
-
-	// Some players like Electra report that they are closed at this point, so wait a frame.
-	bGetDurationDelay = true;
-}
-
-bool FMediaPlateTrackEditor::GetDuration(
-	TStrongObjectPtr<UMediaPlayer>& MediaPlayer, TWeakObjectPtr<UMovieSceneSection>& NewSection)
-{
-	bool bIsDone = false;
-	
-	// Check everything is ok.
-	if ((MediaPlayer.IsValid() == false) || (MediaPlayer->HasError()) || (MediaPlayer->IsClosed()) ||
-		(NewSection.IsValid() == false))
-	{
-		bIsDone = true;
-	}
-	else
-	{
-		// Get the duration.
-		FTimespan Duration = MediaPlayer->GetDuration();
-		if (Duration != 0)
-		{
-			// Once it is non zero, then set the length of the section.
-			FFrameRate TickResolution = NewSection->GetTypedOuter<UMovieScene>()->GetTickResolution();
-			FFrameNumber StartFrame = NewSection->GetInclusiveStartFrame();
-			FFrameNumber EndFrame = StartFrame + (Duration.GetTotalSeconds() * TickResolution).FrameNumber;
-			NewSection->SetEndFrame(TRangeBound<FFrameNumber>::Exclusive(EndFrame));
-			bIsDone = true;
-		}
-	}
-
-	return bIsDone;
-}
-
 void FMediaPlateTrackEditor::OnBuildOutlinerEditWidget(FMenuBuilder& MenuBuilder)
 {
 	// Add media plate sub menu.
diff --git a/Engine/Plugins/Media/MediaPlate/Source/MediaPlateEditor/Private/Sequencer/MediaPlateTrackEditor.h b/Engine/Plugins/Media/MediaPlate/Source/MediaPlateEditor/Private/Sequencer/MediaPlateTrackEditor.h
--- a/Engine/Plugins/Media/MediaPlate/Source/MediaPlateEditor/Private/Sequencer/MediaPlateTrackEditor.h
+++ b/Engine/Plugins/Media/MediaPlate/Source/MediaPlateEditor/Private/Sequencer/MediaPlateTrackEditor.h
@@ -10,8 +10,6 @@
 class AActor;
 class ISequencer;
 class UMediaPlateComponent;
-class UMediaPlayer;
-class UMediaSource;
 
 /**
  * Track editor for media plate.
@@ -57,7 +55,6 @@ public:
 	virtual void BuildObjectBindingContextMenu(FMenuBuilder& MenuBuilder, const TArray<FGuid>& ObjectBindings, const UClass* ObjectClass) override;
 	virtual void BuildObjectBindingTrackMenu(FMenuBuilder& MenuBuilder, const TArray<FGuid>& ObjectBindings, const UClass* ObjectClass) override;
 	virtual bool SupportsType(TSubclassOf<UMovieSceneTrack> TrackClass) const override;
-	virtual void Tick(float DeltaTime) override;
 	//~ ISequencerTrackEditor interface
 	virtual void OnRelease() override;
 
@@ -80,25 +77,6 @@ private:
 	 */
 	void ImportObjectBinding(const TArray<FGuid> ObjectBindings);
 
-	/**
-	 * Starts the process to get the duration of the media.
-	 * It might take a frame or more.
-	 * 
-	 * @param MediaSource		Media to inspect.
-	 * @param Section			Will set this sequencer section to the length of the media.
-	*/
-	void StartGetDuration(UMediaSource* MediaSource, UMovieSceneSection* Section);
-
-	/**
-	 * Call this after StartGetDuration to try and get the duration of the media.
-	 *
-	 * @param MediaPlayer		Player that is opening the media.
-	 * @param NewSection		Movie section this is for.
-	 * @return True if it is done and the player can be removed.
-	 */
-	bool GetDuration(TStrongObjectPtr<UMediaPlayer>& MediaPlayer,
-		TWeakObjectPtr<UMovieSceneSection>& NewSection);
-
 	/**
 	 * Adds media plates to the MediaTrackEditor widget.
 	 * 
@@ -115,10 +93,4 @@ private:
 
 	/** Handle to our delegate. */
 	FDelegateHandle OnActorAddedToSequencerHandle;
-
-	/** List of players that are we are trying to get durations from for the corresponding sections. */
-	TArray<TPair<TStrongObjectPtr<UMediaPlayer>, TWeakObjectPtr<UMovieSceneSection>>> NewSections;
-
-	/** If true then don't check for the duration this frame. */
-	bool bGetDurationDelay;
 };
