MediaCompositing: LWC truncation fixes.
#rnx

diff --git a/Engine/Plugins/Media/MediaCompositing/Source/MediaCompositing/Private/MovieScene/MovieSceneMediaTemplate.cpp b/Engine/Plugins/Media/MediaCompositing/Source/MediaCompositing/Private/MovieScene/MovieSceneMediaTemplate.cpp
--- a/Engine/Plugins/Media/MediaCompositing/Source/MediaCompositing/Private/MovieScene/MovieSceneMediaTemplate.cpp
+++ b/Engine/Plugins/Media/MediaCompositing/Source/MediaCompositing/Private/MovieScene/MovieSceneMediaTemplate.cpp
@@ -333,7 +333,7 @@ void FMovieSceneMediaSectionTemplate::Evaluate(const FMovieSceneEvaluationOperan
 		const FFrameNumber StartFrame = Context.HasPreRollEndTime() ? Context.GetPreRollEndFrame() - Params.SectionStartFrame + Params.StartFrameOffset : Params.StartFrameOffset;
 
 		const double StartFrameInSeconds = FrameRate.AsSeconds(StartFrame);
-		const int64 StartTicks = StartFrameInSeconds * ETimespan::TicksPerSecond;
+		const int64 StartTicks = static_cast<int64>(StartFrameInSeconds * ETimespan::TicksPerSecond);
 
 		ExecutionTokens.Add(FMediaSectionPreRollExecutionToken(MediaSource, Params.MediaSourceProxy, Params.MediaSourceProxyIndex, FTimespan(StartTicks)));
 	}
@@ -343,11 +343,11 @@ void FMovieSceneMediaSectionTemplate::Evaluate(const FMovieSceneEvaluationOperan
 		const FFrameTime FrameTime(Context.GetTime().FrameNumber - Params.SectionStartFrame + Params.StartFrameOffset);
 
 		const double FrameTimeInSeconds = FrameRate.AsSeconds(FrameTime);
-		const int64 FrameTicks = FrameTimeInSeconds * ETimespan::TicksPerSecond;
+		const int64 FrameTicks = static_cast<int64>(FrameTimeInSeconds * ETimespan::TicksPerSecond);
 
 		// With zero-length frames (which can occur occasionally), we use the fixed frame time, matching previous behavior.
 		const double FrameDurationInSeconds = FMath::Max(FrameRate.AsSeconds(FFrameTime(1)), (Context.GetRange().Size<FFrameTime>()) / Context.GetFrameRate());
-		const int64 FrameDurationTicks = FrameDurationInSeconds * ETimespan::TicksPerSecond;
+		const int64 FrameDurationTicks = static_cast<int64>(FrameDurationInSeconds * ETimespan::TicksPerSecond);
 
 		float ProxyTextureBlend = MediaSection->EvaluateEasing(Context.GetTime());
 
diff --git a/Engine/Plugins/Media/MediaCompositing/Source/MediaCompositingEditor/Private/Sequencer/MediaSequenceRecorderExtender.cpp b/Engine/Plugins/Media/MediaCompositing/Source/MediaCompositingEditor/Private/Sequencer/MediaSequenceRecorderExtender.cpp
--- a/Engine/Plugins/Media/MediaCompositing/Source/MediaCompositingEditor/Private/Sequencer/MediaSequenceRecorderExtender.cpp
+++ b/Engine/Plugins/Media/MediaCompositing/Source/MediaCompositingEditor/Private/Sequencer/MediaSequenceRecorderExtender.cpp
@@ -70,7 +70,7 @@ public:
 		{
 			return
 				SNew(SButton)
-				.ContentPadding(0)
+				.ContentPadding(0.0f)
 				.OnClicked(this, &SSequenceRecorderMediaPlayerListRow::ToggleRecordingActive)
 				.ButtonStyle(FAppStyle::Get(), "NoBorder")
 				.ToolTipText(LOCTEXT("ActiveButtonToolTip", "Toggle Recording Active"))
@@ -86,7 +86,7 @@ public:
 		{
 			return
 				SNew(SButton)
-				.ContentPadding(0)
+				.ContentPadding(0.0f)
 				.OnClicked(this, &SSequenceRecorderMediaPlayerListRow::ToggleRecordingFrame)
 				.ButtonStyle(FAppStyle::Get(), "NoBorder")
 				.ToolTipText(LOCTEXT("VideoFramesButtonToolTip", "Toggle Recording Video Frames\nSelect Media Player to show Media Player Recording Options"))
diff --git a/Engine/Plugins/Media/MediaCompositing/Source/MediaCompositingEditor/Private/Sequencer/MediaThumbnailSection.cpp b/Engine/Plugins/Media/MediaCompositing/Source/MediaCompositingEditor/Private/Sequencer/MediaThumbnailSection.cpp
--- a/Engine/Plugins/Media/MediaCompositing/Source/MediaCompositingEditor/Private/Sequencer/MediaThumbnailSection.cpp
+++ b/Engine/Plugins/Media/MediaCompositing/Source/MediaCompositingEditor/Private/Sequencer/MediaThumbnailSection.cpp
@@ -162,7 +162,7 @@ void FMediaThumbnailSection::SetSingleTime(double GlobalTime)
 	if (MediaSection != nullptr)
 	{
 		double StartTime = MediaSection->GetInclusiveStartFrame() / MediaSection->GetTypedOuter<UMovieScene>()->GetTickResolution();
-		MediaSection->SetThumbnailReferenceOffset(GlobalTime - StartTime);
+		MediaSection->SetThumbnailReferenceOffset(static_cast<float>(GlobalTime - StartTime));
 	}
 }
 
@@ -297,7 +297,7 @@ void FMediaThumbnailSection::DrawLoopIndicators(FSequencerSectionPainter& InPain
 
 	FFrameRate TickResolution = Section->GetTypedOuter<UMovieScene>()->GetTickResolution();
 	double SectionDuration = FFrameTime(UE::MovieScene::DiscreteSize(Section->GetRange())) / TickResolution;
-	const float MediaSizeX = MediaDuration.GetTotalSeconds() * SectionSize.X / SectionDuration;
+	const float MediaSizeX = static_cast<float>(MediaDuration.GetTotalSeconds() * SectionSize.X / SectionDuration);
 	const FFrameNumber SectionOffset = MediaSection->GetRange().HasLowerBound() ? MediaSection->GetRange().GetLowerBoundValue() : 0;
 	float DrawOffset = MediaSizeX - TimeToPixelConverter.SecondsToPixel(TickResolution.AsSeconds(SectionOffset + MediaSection->StartFrameOffset));
 
@@ -329,7 +329,7 @@ void FMediaThumbnailSection::DrawSampleStates(FSequencerSectionPainter& InPainte
 
 	FFrameRate TickResolution = Section->GetTypedOuter<UMovieScene>()->GetTickResolution();
 	double SectionDuration = FFrameTime(UE::MovieScene::DiscreteSize(Section->GetRange())) / TickResolution;
-	const float MediaSizeX = MediaDuration.GetTotalSeconds() * SectionSize.X / SectionDuration;
+	const float MediaSizeX = static_cast<float>(MediaDuration.GetTotalSeconds() * SectionSize.X / SectionDuration);
 
 	TArray<TRange<FTimespan>> Ranges;
 	RangeSet.GetRanges(Ranges);
@@ -339,10 +339,10 @@ void FMediaThumbnailSection::DrawSampleStates(FSequencerSectionPainter& InPainte
 	{
 		for (auto& Range : Ranges)
 		{
-			const float DrawOffset = FMath::RoundToNegativeInfinity(FTimespan::Ratio(Range.GetLowerBoundValue(), MediaDuration) * MediaSizeX) +	
+			const float DrawOffset = static_cast<float>(FMath::RoundToNegativeInfinity(FTimespan::Ratio(Range.GetLowerBoundValue(), MediaDuration) * MediaSizeX) +	
 				LoopDrawOffset +
-				TimeToPixelConverter.SecondsToPixel(0.0);
-			const float DrawSize = FMath::RoundToPositiveInfinity(FTimespan::Ratio(Range.Size<FTimespan>(), MediaDuration) * MediaSizeX);
+				TimeToPixelConverter.SecondsToPixel(0.0));
+			const float DrawSize = static_cast<float>(FMath::RoundToPositiveInfinity(FTimespan::Ratio(Range.Size<FTimespan>(), MediaDuration) * MediaSizeX));
 			const float BarHeight = 4.0f;
 
 			FSlateDrawElement::MakeBox(
@@ -386,7 +386,7 @@ void FMediaThumbnailSection::DrawMediaInfo(FSequencerSectionPainter& InPainter,
 	const FSlateFontInfo SmallLayoutFont = FCoreStyle::GetDefaultFontStyle("Bold", 10);
 	const TSharedRef< FSlateFontMeasure > FontMeasureService = FSlateApplication::Get().GetRenderer()->GetFontMeasureService();
 	
-	FVector2D TextSize(EForceInit::ForceInitToZero);
+	FVector2f TextSize(EForceInit::ForceInitToZero);
 	FMargin ContentPadding = GetContentPadding();
 	FVector2D TextOffset(EForceInit::ForceInitToZero);
 	float BaseYOffset = 0.0f;
