Media: Fix missing transaction for delayed media section duration update.

diff --git a/Engine/Plugins/Media/MediaCompositing/Source/MediaCompositing/MediaCompositing.Build.cs b/Engine/Plugins/Media/MediaCompositing/Source/MediaCompositing/MediaCompositing.Build.cs
--- a/Engine/Plugins/Media/MediaCompositing/Source/MediaCompositing/MediaCompositing.Build.cs
+++ b/Engine/Plugins/Media/MediaCompositing/Source/MediaCompositing/MediaCompositing.Build.cs
@@ -16,6 +16,11 @@ namespace UnrealBuildTool.Rules
 					"TimeManagement",
 				});
 
+			if (Target.bBuildEditor)
+			{
+				PrivateDependencyModuleNames.Add("UnrealEd");
+			}
+
 			PrivateDependencyModuleNames.AddRange(
 				new string[] {
 					"Core",
diff --git a/Engine/Plugins/Media/MediaCompositing/Source/MediaCompositing/Private/MovieScene/MovieSceneMediaTrack.cpp b/Engine/Plugins/Media/MediaCompositing/Source/MediaCompositing/Private/MovieScene/MovieSceneMediaTrack.cpp
--- a/Engine/Plugins/Media/MediaCompositing/Source/MediaCompositing/Private/MovieScene/MovieSceneMediaTrack.cpp
+++ b/Engine/Plugins/Media/MediaCompositing/Source/MediaCompositing/Private/MovieScene/MovieSceneMediaTrack.cpp
@@ -14,6 +14,10 @@
 #include "UObject/Package.h"
 #include "UObject/UObjectGlobals.h"
 
+#if WITH_EDITOR
+#include "ScopedTransaction.h"
+#endif
+
 #include UE_INLINE_GENERATED_CPP_BY_NAME(MovieSceneMediaTrack)
 
 
@@ -122,6 +126,8 @@ UMovieSceneMediaSection* UMovieSceneMediaTrack::AddNewSectionOnRow(FFrameNumber
 void UMovieSceneMediaTrack::TickOutput()
 {
 #if WITH_EDITOR
+	const FScopedTransaction Transaction(NSLOCTEXT("Sequencer", "MediaSectionDuration_Transaction", "Set Media Section Duration"));
+
 	// Do we have any new sections that need durations?
 	if (NewSections.Num() > 0)
 	{
