MediaCompositing: Added key to allow the player to be closed.

diff --git a/Engine/Plugins/Media/MediaCompositing/Source/MediaCompositing/Private/MovieScene/MovieSceneMediaSection.cpp b/Engine/Plugins/Media/MediaCompositing/Source/MediaCompositing/Private/MovieScene/MovieSceneMediaSection.cpp
--- a/Engine/Plugins/Media/MediaCompositing/Source/MediaCompositing/Private/MovieScene/MovieSceneMediaSection.cpp
+++ b/Engine/Plugins/Media/MediaCompositing/Source/MediaCompositing/Private/MovieScene/MovieSceneMediaSection.cpp
@@ -36,6 +36,14 @@ UMovieSceneMediaSection::UMovieSceneMediaSection(const FObjectInitializer& Objec
 
 	EvalOptions.CompletionMode = EMovieSceneCompletionMode::RestoreState;
 	BlendType = EMovieSceneBlendType::Absolute;
+	ChannelCanPlayerBeOpen.SetDefault(true);
+
+#if WITH_EDITOR
+	static const FMovieSceneChannelMetaData MetaData("CanPlayerBeOpen", LOCTEXT("CanPlayerBeOpenChannelText", "Can the media player be open."));
+	ChannelProxy = MakeShared<FMovieSceneChannelProxy>(ChannelCanPlayerBeOpen, MetaData, TMovieSceneExternalValue<bool>::Make());
+#else
+	ChannelProxy = MakeShared<FMovieSceneChannelProxy>(ChannelCanPlayerBeOpen);
+#endif
 }
 
 void UMovieSceneMediaSection::PostInitProperties()
diff --git a/Engine/Plugins/Media/MediaCompositing/Source/MediaCompositing/Private/MovieScene/MovieSceneMediaTemplate.cpp b/Engine/Plugins/Media/MediaCompositing/Source/MediaCompositing/Private/MovieScene/MovieSceneMediaTemplate.cpp
--- a/Engine/Plugins/Media/MediaCompositing/Source/MediaCompositing/Private/MovieScene/MovieSceneMediaTemplate.cpp
+++ b/Engine/Plugins/Media/MediaCompositing/Source/MediaCompositing/Private/MovieScene/MovieSceneMediaTemplate.cpp
@@ -114,12 +114,13 @@ private:
 struct FMediaSectionExecutionToken
 	: FMediaSectionBaseExecutionToken
 {
-	FMediaSectionExecutionToken(UMediaSource* InMediaSource, FMovieSceneObjectBindingID InMediaSourceProxy, int32 InMediaSourceProxyIndex, float InProxyTextureBlend, FTimespan InCurrentTime, FTimespan InFrameDuration)
+	FMediaSectionExecutionToken(UMediaSource* InMediaSource, FMovieSceneObjectBindingID InMediaSourceProxy, int32 InMediaSourceProxyIndex, float InProxyTextureBlend, bool bInCanPlayerBeOpen, FTimespan InCurrentTime, FTimespan InFrameDuration)
 		: FMediaSectionBaseExecutionToken(InMediaSource, InMediaSourceProxy, InMediaSourceProxyIndex)
 		, CurrentTime(InCurrentTime)
 		, FrameDuration(InFrameDuration)
 		, PlaybackRate(1.0f)
 		, ProxyTextureBlend(InProxyTextureBlend)
+		, bCanPlayerBeOpen(bInCanPlayerBeOpen)
 	{ }
 
 	virtual void Execute(const FMovieSceneContext& Context, const FMovieSceneEvaluationOperand& Operand, FPersistentEvaluationData& PersistentData, IMovieScenePlayer& Player) override
@@ -159,6 +160,16 @@ struct FMediaSectionExecutionToken
 			}
 		}
 
+		// Can we be open?
+		if (bCanPlayerBeOpen == false)
+		{
+			if (MediaPlayer->IsClosed() == false)
+			{
+				MediaPlayer->Close();
+			}
+			return;
+		}
+
 		// open the media source if necessary
 		if (MediaPlayer->GetUrl().IsEmpty())
 		{
@@ -269,6 +280,7 @@ private:
 	FTimespan FrameDuration;
 	float PlaybackRate;
 	float ProxyTextureBlend;
+	bool bCanPlayerBeOpen;
 };
 
 
@@ -340,6 +352,8 @@ void FMovieSceneMediaSectionTemplate::Evaluate(const FMovieSceneEvaluationOperan
 		const int64 FrameDurationTicks = static_cast<int64>(FrameDurationInSeconds * ETimespan::TicksPerSecond);
 
 		float ProxyTextureBlend = MediaSection->EvaluateEasing(Context.GetTime());
+		bool bCanPlayerBeOpen = true;
+		MediaSection->ChannelCanPlayerBeOpen.Evaluate(Context.GetTime(), bCanPlayerBeOpen);
 
 		#if MOVIESCENEMEDIATEMPLATE_TRACE_EVALUATION
 			GLog->Logf(ELogVerbosity::Log, TEXT("Evaluating frame %i+%f, FrameRate %i/%i, FrameTicks %d, FrameDurationTicks %d"),
@@ -352,7 +366,7 @@ void FMovieSceneMediaSectionTemplate::Evaluate(const FMovieSceneEvaluationOperan
 			);
 		#endif
 
-		ExecutionTokens.Add(FMediaSectionExecutionToken(MediaSource, Params.MediaSourceProxy, Params.MediaSourceProxyIndex, ProxyTextureBlend, FTimespan(FrameTicks), FTimespan(FrameDurationTicks)));
+		ExecutionTokens.Add(FMediaSectionExecutionToken(MediaSource, Params.MediaSourceProxy, Params.MediaSourceProxyIndex, ProxyTextureBlend, bCanPlayerBeOpen, FTimespan(FrameTicks), FTimespan(FrameDurationTicks)));
 	}
 }
 
diff --git a/Engine/Plugins/Media/MediaCompositing/Source/MediaCompositing/Public/MovieSceneMediaSection.h b/Engine/Plugins/Media/MediaCompositing/Source/MediaCompositing/Public/MovieSceneMediaSection.h
--- a/Engine/Plugins/Media/MediaCompositing/Source/MediaCompositing/Public/MovieSceneMediaSection.h
+++ b/Engine/Plugins/Media/MediaCompositing/Source/MediaCompositing/Public/MovieSceneMediaSection.h
@@ -4,8 +4,10 @@
 
 #include "MovieSceneSection.h"
 
+#include "Channels/MovieSceneChannelProxy.h"
 #include "MediaSource.h"
 #include "MovieSceneObjectBindingID.h"
+#include "Channels/MovieSceneBoolChannel.h"
 
 #include "MovieSceneMediaSection.generated.h"
 
@@ -73,6 +75,10 @@ public:
 	UPROPERTY(Transient)
 	bool bHasMediaPlayerProxy;
 
+	/** If true then the media player can be open. */
+	UPROPERTY()
+	FMovieSceneBoolChannel ChannelCanPlayerBeOpen;
+
 public:
 
 	/**
