MediaCompositing: Added thumbnails to the track, but currently needs a pre existing media source thumbnail.

#rnx

diff --git a/Engine/Plugins/Media/MediaCompositing/Source/MediaCompositingEditor/MediaCompositingEditor.Build.cs b/Engine/Plugins/Media/MediaCompositing/Source/MediaCompositingEditor/MediaCompositingEditor.Build.cs
--- a/Engine/Plugins/Media/MediaCompositing/Source/MediaCompositingEditor/MediaCompositingEditor.Build.cs
+++ b/Engine/Plugins/Media/MediaCompositing/Source/MediaCompositingEditor/MediaCompositingEditor.Build.cs
@@ -22,6 +22,7 @@ namespace UnrealBuildTool.Rules
 					"MediaUtils",
 					"MovieScene",
 					"MovieSceneTracks",
+					"RenderCore",
 					"RHI",
 					"Sequencer",
 					"SequenceRecorder",
diff --git a/Engine/Plugins/Media/MediaCompositing/Source/MediaCompositingEditor/Private/Sequencer/MediaThumbnailSection.cpp b/Engine/Plugins/Media/MediaCompositing/Source/MediaCompositingEditor/Private/Sequencer/MediaThumbnailSection.cpp
--- a/Engine/Plugins/Media/MediaCompositing/Source/MediaCompositingEditor/Private/Sequencer/MediaThumbnailSection.cpp
+++ b/Engine/Plugins/Media/MediaCompositing/Source/MediaCompositingEditor/Private/Sequencer/MediaThumbnailSection.cpp
@@ -2,6 +2,7 @@
 
 #include "Sequencer/MediaThumbnailSection.h"
 
+#include "CommonRenderResources.h"
 #include "Fonts/FontMeasure.h"
 #include "Framework/Application/SlateApplication.h"
 #include "Styling/AppStyle.h"
@@ -12,10 +13,13 @@
 #include "MediaPlayerFacade.h"
 #include "MediaSource.h"
 #include "MediaTexture.h"
+#include "Modules/ModuleManager.h"
 #include "MovieScene.h"
 #include "MovieSceneMediaSection.h"
 #include "MovieSceneTimeHelpers.h"
 #include "Rendering/DrawElements.h"
+#include "RHIStaticStates.h"
+#include "ScreenRendering.h"
 #include "SequencerSectionPainter.h"
 #include "TrackEditorThumbnail/TrackEditorThumbnailPool.h"
 #include "Compilation/MovieSceneCompiledDataManager.h"
@@ -71,14 +75,7 @@ float FMediaThumbnailSection::GetSectionHeight() const
 
 FText FMediaThumbnailSection::GetSectionTitle() const
 {
-	UMovieSceneMediaSection* MediaSection = CastChecked<UMovieSceneMediaSection>(Section);
-	TSharedPtr<ISequencer> Sequencer = SequencerPtr.Pin();
-	UMediaSource* MediaSource = nullptr;
-
-	if (MediaSection && Sequencer.IsValid())
-	{
-		MediaSource = MediaSection->GetMediaSourceOrProxy(*Sequencer, Sequencer->GetFocusedTemplateID());
-	}
+	UMediaSource* MediaSource = GetMediaSource();
 
 	if (MediaSource == nullptr)
 	{
@@ -250,6 +247,32 @@ void FMediaThumbnailSection::SlipSection(FFrameNumber SlipTime)
 
 void FMediaThumbnailSection::Draw(FTrackEditorThumbnail& TrackEditorThumbnail)
 {
+	UMediaSource* MediaSource = GetMediaSource();
+	UTexture* Thumbnail = MediaSource->GetThumbnail();
+	if (Thumbnail != nullptr)
+	{
+		FTextureReferenceRHIRef SourceTexture = Thumbnail->TextureReference.TextureReferenceRHI;
+		if (SourceTexture.IsValid())
+		{
+			// Limit thumbnail size.
+			FIntPoint RTSize(SourceTexture->GetDesc().Extent);
+			int32 SourceMaxSize = RTSize.GetMax();
+			int32 ThumbnailSize = 256;
+			if (ThumbnailSize < SourceMaxSize)
+			{
+				RTSize = ((RTSize * ThumbnailSize) / SourceMaxSize);
+			}
+
+			TrackEditorThumbnail.bIgnoreAlpha = true;
+			TrackEditorThumbnail.ResizeRenderTarget(RTSize);
+			FSlateTextureRenderTarget2DResource* RenderTarget =
+				TrackEditorThumbnail.GetRenderTarget();
+			if (RenderTarget != nullptr)
+			{
+				CopyTexture(RenderTarget, SourceTexture);
+			}
+		}
+	}
 }
 
 
@@ -430,6 +453,20 @@ void FMediaThumbnailSection::DrawMediaInfo(FSequencerSectionPainter& InPainter,
 	}
 }
 
+UMediaSource* FMediaThumbnailSection::GetMediaSource() const
+{
+	UMovieSceneMediaSection* MediaSection = CastChecked<UMovieSceneMediaSection>(Section);
+	TSharedPtr<ISequencer> Sequencer = SequencerPtr.Pin();
+	UMediaSource* MediaSource = nullptr;
+
+	if (MediaSection && Sequencer.IsValid())
+	{
+		MediaSource = MediaSection->GetMediaSourceOrProxy(*Sequencer, Sequencer->GetFocusedTemplateID());
+	}
+
+	return MediaSource;
+}
+
 UMediaPlayer* FMediaThumbnailSection::GetTemplateMediaPlayer() const
 {
 	// locate the track that evaluates this section
@@ -506,5 +543,64 @@ UMediaPlayer* FMediaThumbnailSection::GetTemplateMediaPlayer() const
 	return MediaData->GetMediaPlayer();
 }
 
+void FMediaThumbnailSection::CopyTexture(FSlateTextureRenderTarget2DResource* RenderTarget, FTextureReferenceRHIRef SourceTexture)
+{
+	ENQUEUE_RENDER_COMMAND(MediaThumbnailCopyTexture)(
+		[SourceTexture, RenderTarget](FRHICommandListImmediate& RHICmdList)
+		{
+			IRendererModule* RendererModule = &FModuleManager::GetModuleChecked<IRendererModule>("Renderer");
+
+			FTextureRHIRef TargetTexture = RenderTarget->GetRenderTargetTexture();
+			if (TargetTexture.IsValid())
+			{
+				RHICmdList.Transition(FRHITransitionInfo(TargetTexture, ERHIAccess::Unknown, ERHIAccess::RTV));
+
+				FRHIRenderPassInfo RPInfo(TargetTexture, ERenderTargetActions::Load_Store);
+				RHICmdList.BeginRenderPass(RPInfo, TEXT("MediaThumbnailCopyTexture"));
+				{
+					RHICmdList.SetViewport(0.0f, 0.0f, 0.0f,
+						static_cast<float>(TargetTexture->GetSizeX()),
+						static_cast<float>(TargetTexture->GetSizeY()), 1.0f);
+
+					FGraphicsPipelineStateInitializer GraphicsPSOInit;
+					RHICmdList.ApplyCachedRenderTargets(GraphicsPSOInit);
+					GraphicsPSOInit.BlendState = TStaticBlendState<>::GetRHI();
+					GraphicsPSOInit.RasterizerState = TStaticRasterizerState<>::GetRHI();
+					GraphicsPSOInit.DepthStencilState = TStaticDepthStencilState<false, CF_Always>::GetRHI();
+
+					FGlobalShaderMap* ShaderMap = GetGlobalShaderMap(GMaxRHIFeatureLevel);
+					TShaderMapRef<FScreenVS> VertexShader(ShaderMap);
+					TShaderMapRef<FScreenPS> PixelShader(ShaderMap);
+
+					GraphicsPSOInit.BoundShaderState.VertexDeclarationRHI = GFilterVertexDeclaration.VertexDeclarationRHI;
+					GraphicsPSOInit.BoundShaderState.VertexShaderRHI = VertexShader.GetVertexShader();
+					GraphicsPSOInit.BoundShaderState.PixelShaderRHI = PixelShader.GetPixelShader();
+
+					GraphicsPSOInit.PrimitiveType = PT_TriangleList;
+
+					SetGraphicsPipelineState(RHICmdList, GraphicsPSOInit, 0);
+
+					const bool bSameSize = (TargetTexture->GetDesc().Extent == SourceTexture->GetDesc().Extent);
+					FRHISamplerState* PixelSampler = bSameSize ? TStaticSamplerState<SF_Point>::GetRHI() : TStaticSamplerState<SF_Bilinear>::GetRHI();
+
+					SetShaderParametersLegacyPS(RHICmdList, PixelShader, PixelSampler, SourceTexture);
+
+					RendererModule->DrawRectangle(
+						RHICmdList,
+						0.0f, 0.0f,					// Dest X, Y
+						static_cast<float>(TargetTexture->GetSizeX()),	// Dest Width
+						static_cast<float>(TargetTexture->GetSizeY()),	// Dest Height
+						0, 0,						// Source U, V
+						1, 1,						// Source USize, VSize
+						TargetTexture->GetSizeXY(),	// Target buffer size
+						FIntPoint(1, 1),			// Source texture size
+						VertexShader,
+						EDRF_Default);
+				}
+				RHICmdList.EndRenderPass();
+				RHICmdList.Transition(FRHITransitionInfo(TargetTexture, ERHIAccess::RTV, ERHIAccess::SRVMask));
+			}
+		});
+}
 
 #undef LOCTEXT_NAMESPACE
diff --git a/Engine/Plugins/Media/MediaCompositing/Source/MediaCompositingEditor/Public/Sequencer/MediaThumbnailSection.h b/Engine/Plugins/Media/MediaCompositing/Source/MediaCompositingEditor/Public/Sequencer/MediaThumbnailSection.h
--- a/Engine/Plugins/Media/MediaCompositing/Source/MediaCompositingEditor/Public/Sequencer/MediaThumbnailSection.h
+++ b/Engine/Plugins/Media/MediaCompositing/Source/MediaCompositingEditor/Public/Sequencer/MediaThumbnailSection.h
@@ -4,14 +4,17 @@
 
 #include "Math/RangeSet.h"
 #include "Misc/Timespan.h"
+#include "RHIFwd.h"
 #include "Sections/ThumbnailSection.h"
 #include "Templates/SharedPointer.h"
 #include "TrackEditorThumbnail/TrackEditorThumbnail.h"
 #include "UObject/GCObject.h"
 
+class FSlateTextureRenderTarget2DResource;
 class FTrackEditorThumbnailPool;
 class ISequencer;
 class UMediaPlayer;
+class UMediaSource;
 class UMediaTexture;
 class UMovieSceneMediaSection;
 
@@ -103,6 +106,11 @@ protected:
 	 */
 	void DrawMediaInfo(FSequencerSectionPainter& InPainter, UMediaPlayer* MediaPlayer, FVector2D SectionSize) const;
 
+	/**
+	 * Helper function to get the media source.
+	 */
+	UMediaSource* GetMediaSource() const;
+
 	/**
 	 * Get the media player that is used by the evaluation template.
 	 *
@@ -110,6 +118,11 @@ protected:
 	 */
 	UMediaPlayer* GetTemplateMediaPlayer() const;
 
+	/**
+	 * Helper function to copy a texture.
+	 */
+	void CopyTexture(FSlateTextureRenderTarget2DResource* RenderTarget, FTextureReferenceRHIRef SourceTexture);
+
 private:
 
 	/** The section object that owns this section. */
