MediaCompositing: Fixed thumbnail crash when there is no media source.
#rnx

diff --git a/Engine/Plugins/Media/MediaCompositing/Source/MediaCompositingEditor/Private/Sequencer/MediaThumbnailSection.cpp b/Engine/Plugins/Media/MediaCompositing/Source/MediaCompositingEditor/Private/Sequencer/MediaThumbnailSection.cpp
--- a/Engine/Plugins/Media/MediaCompositing/Source/MediaCompositingEditor/Private/Sequencer/MediaThumbnailSection.cpp
+++ b/Engine/Plugins/Media/MediaCompositing/Source/MediaCompositingEditor/Private/Sequencer/MediaThumbnailSection.cpp
@@ -248,28 +248,31 @@ void FMediaThumbnailSection::SlipSection(FFrameNumber SlipTime)
 void FMediaThumbnailSection::Draw(FTrackEditorThumbnail& TrackEditorThumbnail)
 {
 	UMediaSource* MediaSource = GetMediaSource();
-	UTexture* Thumbnail = MediaSource->GetThumbnail();
-	if (Thumbnail != nullptr)
+	if (MediaSource != nullptr)
 	{
-		FTextureReferenceRHIRef SourceTexture = Thumbnail->TextureReference.TextureReferenceRHI;
-		if (SourceTexture.IsValid())
+		UTexture* Thumbnail = MediaSource->GetThumbnail();
+		if (Thumbnail != nullptr)
 		{
-			// Limit thumbnail size.
-			FIntPoint RTSize(SourceTexture->GetDesc().Extent);
-			int32 SourceMaxSize = RTSize.GetMax();
-			int32 ThumbnailSize = 256;
-			if (ThumbnailSize < SourceMaxSize)
+			FTextureReferenceRHIRef SourceTexture = Thumbnail->TextureReference.TextureReferenceRHI;
+			if (SourceTexture.IsValid())
 			{
-				RTSize = ((RTSize * ThumbnailSize) / SourceMaxSize);
-			}
+				// Limit thumbnail size.
+				FIntPoint RTSize(SourceTexture->GetDesc().Extent);
+				int32 SourceMaxSize = RTSize.GetMax();
+				int32 ThumbnailSize = 256;
+				if (ThumbnailSize < SourceMaxSize)
+				{
+					RTSize = ((RTSize * ThumbnailSize) / SourceMaxSize);
+				}
 
-			TrackEditorThumbnail.bIgnoreAlpha = true;
-			TrackEditorThumbnail.ResizeRenderTarget(RTSize);
-			FSlateTextureRenderTarget2DResource* RenderTarget =
-				TrackEditorThumbnail.GetRenderTarget();
-			if (RenderTarget != nullptr)
-			{
-				CopyTexture(RenderTarget, SourceTexture);
+				TrackEditorThumbnail.bIgnoreAlpha = true;
+				TrackEditorThumbnail.ResizeRenderTarget(RTSize);
+				FSlateTextureRenderTarget2DResource* RenderTarget =
+					TrackEditorThumbnail.GetRenderTarget();
+				if (RenderTarget != nullptr)
+				{
+					CopyTexture(RenderTarget, SourceTexture);
+				}
 			}
 		}
 	}
