Media: Restart sequencer playback on cache settings change.

This follows original support added via the player proxy in UE-164195.

diff --git a/Engine/Plugins/Media/MediaCompositing/Source/MediaCompositing/Private/MovieScene/MovieSceneMediaTemplate.cpp b/Engine/Plugins/Media/MediaCompositing/Source/MediaCompositing/Private/MovieScene/MovieSceneMediaTemplate.cpp
--- a/Engine/Plugins/Media/MediaCompositing/Source/MediaCompositing/Private/MovieScene/MovieSceneMediaTemplate.cpp
+++ b/Engine/Plugins/Media/MediaCompositing/Source/MediaCompositing/Private/MovieScene/MovieSceneMediaTemplate.cpp
@@ -76,7 +76,7 @@ struct FMediaSectionPreRollExecutionToken
 
 		FMovieSceneMediaData& SectionData = PersistentData.GetSectionData<FMovieSceneMediaData>();
 		UMediaPlayer* MediaPlayer = SectionData.GetMediaPlayer();
-		UObject* PlayerProxy = SectionData.GetPlayerProxy();
+		IMediaPlayerProxyInterface* PlayerProxyInterface = Cast<IMediaPlayerProxyInterface>(SectionData.GetPlayerProxy());
 		UMediaSource* MediaSource = GetMediaSource(Player, Operand.SequenceID);
 
 		if (MediaPlayer == nullptr || MediaSource == nullptr)
@@ -84,14 +84,6 @@ struct FMediaSectionPreRollExecutionToken
 			return;
 		}
 
-		// Do we have a player proxy?
-		IMediaPlayerProxyInterface* PlayerProxyInterface = nullptr;
-		if (PlayerProxy != nullptr)
-		{
-			PlayerProxyInterface = Cast<IMediaPlayerProxyInterface>
-				(PlayerProxy);
-		}
-
 		// open the media source if necessary
 		if (MediaPlayer->GetUrl().IsEmpty())
 		{
@@ -136,26 +128,21 @@ struct FMediaSectionExecutionToken
 		}
 
 		// Do we have a player proxy?
-		IMediaPlayerProxyInterface* PlayerProxyInterface = nullptr;
-		if (PlayerProxy != nullptr)
+		IMediaPlayerProxyInterface* PlayerProxyInterface = Cast<IMediaPlayerProxyInterface>(PlayerProxy);
+		if (PlayerProxyInterface != nullptr)
 		{
-			PlayerProxyInterface = Cast<IMediaPlayerProxyInterface>
-				(PlayerProxy);
-			if (PlayerProxyInterface != nullptr)
+			PlayerProxyInterface->ProxySetTextureBlend(SectionData.GetProxyLayerIndex(), SectionData.GetProxyTextureIndex(), ProxyTextureBlend);
+			// Can we control the player?
+			if (PlayerProxyInterface->IsExternalControlAllowed() == false)
 			{
-				PlayerProxyInterface->ProxySetTextureBlend(SectionData.GetProxyLayerIndex(), SectionData.GetProxyTextureIndex(), ProxyTextureBlend);
-				// Can we control the player?
-				if (PlayerProxyInterface->IsExternalControlAllowed() == false)
-				{
-					return;
-				}
+				return;
+			}
 
-				if (SectionData.bIsAspectRatioSet == false)
+			if (SectionData.bIsAspectRatioSet == false)
+			{
+				if (PlayerProxyInterface->ProxySetAspectRatio(MediaPlayer))
 				{
-					if (PlayerProxyInterface->ProxySetAspectRatio(MediaPlayer))
-					{
-						SectionData.bIsAspectRatioSet = true;
-					}
+					SectionData.bIsAspectRatioSet = true;
 				}
 			}
 		}
@@ -170,8 +157,15 @@ struct FMediaSectionExecutionToken
 			return;
 		}
 
+		bool bCacheSettingsChanged = false;
+		FMediaSourceCacheSettings CurrentCacheSettings;
+		if (PlayerProxyInterface != nullptr && MediaSource->GetCacheSettings(CurrentCacheSettings))
+		{
+			bCacheSettingsChanged = (CurrentCacheSettings != PlayerProxyInterface->GetCacheSettings());
+		}
+
 		// open the media source if necessary
-		if (MediaPlayer->GetUrl().IsEmpty())
+		if (MediaPlayer->GetUrl().IsEmpty() || bCacheSettingsChanged)
 		{
 			if (PlayerProxyInterface != nullptr)
 			{
diff --git a/Engine/Source/Runtime/MediaAssets/Private/Assets/MediaSource.cpp b/Engine/Source/Runtime/MediaAssets/Private/Assets/MediaSource.cpp
--- a/Engine/Source/Runtime/MediaAssets/Private/Assets/MediaSource.cpp
+++ b/Engine/Source/Runtime/MediaAssets/Private/Assets/MediaSource.cpp
@@ -16,10 +16,30 @@
 
 #include UE_INLINE_GENERATED_CPP_BY_NAME(MediaSource)
 
+static const FLazyName ImgMediaSmartCacheEnabledName(TEXT("ImgMediaSmartCacheEnabled"));
+static const FLazyName ImgMediaSmartCacheTimeToLookAheadName(TEXT("ImgMediaSmartCacheTimeToLookAhead"));
+
 void UMediaSource::SetCacheSettings(const FMediaSourceCacheSettings& Settings)
 {
-	SetMediaOptionBool(TEXT("ImgMediaSmartCacheEnabled"), Settings.bOverride);
-	SetMediaOptionFloat(TEXT("ImgMediaSmartCacheTimeToLookAhead"), Settings.TimeToLookAhead);
+	SetMediaOptionBool(ImgMediaSmartCacheEnabledName, Settings.bOverride);
+	SetMediaOptionFloat(ImgMediaSmartCacheTimeToLookAheadName, Settings.TimeToLookAhead);
+}
+
+bool UMediaSource::GetCacheSettings(FMediaSourceCacheSettings& OutSettings) const
+{
+	if (!HasMediaOption(ImgMediaSmartCacheEnabledName))
+	{
+		return false;
+	}
+
+	if (!HasMediaOption(ImgMediaSmartCacheTimeToLookAheadName))
+	{
+		return false;
+	}
+
+	OutSettings.bOverride = GetMediaOption(ImgMediaSmartCacheEnabledName, false);
+	OutSettings.TimeToLookAhead = GetMediaOption(ImgMediaSmartCacheTimeToLookAheadName, 0.0f);
+	return true;
 }
 
 #if WITH_EDITOR
diff --git a/Engine/Source/Runtime/MediaAssets/Public/MediaSource.h b/Engine/Source/Runtime/MediaAssets/Public/MediaSource.h
--- a/Engine/Source/Runtime/MediaAssets/Public/MediaSource.h
+++ b/Engine/Source/Runtime/MediaAssets/Public/MediaSource.h
@@ -49,6 +49,17 @@ struct FMediaSourceCacheSettings
 	 */
 	UPROPERTY(EditAnywhere, Category = "Media Cache")
 	float TimeToLookAhead = 0.2f;
+
+
+	inline bool operator==(const FMediaSourceCacheSettings& Other) const
+	{
+		return (Other.bOverride == bOverride) && FMath::IsNearlyEqual(Other.TimeToLookAhead, TimeToLookAhead);
+	}
+
+	inline bool operator!=(const FMediaSourceCacheSettings& Other) const
+	{
+		return !(*this == Other);
+	}
 };
 
 /**
@@ -91,6 +102,13 @@ public:
 	 */
 	MEDIAASSETS_API void SetCacheSettings(const FMediaSourceCacheSettings& Settings);
 
+	/**
+	 * Get the media source cache settings, if present.
+	 * 
+	 * @param OutSettings Cache settings
+	 * @return true if the cache settings are present, false otherwise.
+	 */
+	MEDIAASSETS_API bool GetCacheSettings(FMediaSourceCacheSettings& OutSettings) const;
 
 #if WITH_EDITOR
 
