MediaCompositing: Perform fast player teardown only when it was created temporarily, and not if it has been passed in from the outside to be taken control over, as this results in terminal events not being passed out.

#rnx

diff --git a/Engine/Plugins/Media/MediaCompositing/Source/MediaCompositing/Private/MovieScene/MovieSceneMediaData.cpp b/Engine/Plugins/Media/MediaCompositing/Source/MediaCompositing/Private/MovieScene/MovieSceneMediaData.cpp
--- a/Engine/Plugins/Media/MediaCompositing/Source/MediaCompositing/Private/MovieScene/MovieSceneMediaData.cpp
+++ b/Engine/Plugins/Media/MediaCompositing/Source/MediaCompositing/Private/MovieScene/MovieSceneMediaData.cpp
@@ -28,7 +28,10 @@ FMovieSceneMediaData::~FMovieSceneMediaData()
 	{
 		MediaPlayer->OnMediaEvent().RemoveAll(this);
 		MediaPlayer->Close();
-		MediaPlayer->CleanUpBeforeDestroy();
+		if (!bOverrideMediaPlayer)
+		{
+			MediaPlayer->CleanUpBeforeDestroy();
+		}
 		MediaPlayer->RemoveFromRoot();
 	}
 }
