MediaCompositing: Prerolling a media track after a loop requires the player to seek to the new time and removal of the blocked range.

#rnx

diff --git a/Engine/Plugins/Media/MediaCompositing/Source/MediaCompositing/Private/MovieScene/MovieSceneMediaTemplate.cpp b/Engine/Plugins/Media/MediaCompositing/Source/MediaCompositing/Private/MovieScene/MovieSceneMediaTemplate.cpp
--- a/Engine/Plugins/Media/MediaCompositing/Source/MediaCompositing/Private/MovieScene/MovieSceneMediaTemplate.cpp
+++ b/Engine/Plugins/Media/MediaCompositing/Source/MediaCompositing/Private/MovieScene/MovieSceneMediaTemplate.cpp
@@ -94,6 +94,15 @@ struct FMediaSectionPreRollExecutionToken
 			
 			SectionData.SeekOnOpen(StartTime);
 			MediaPlayer->OpenSource(MediaSource);
+			return;
+		}
+
+		bool bMoveToNewTime = Context.GetStatus() != EMovieScenePlayerStatus::Playing || (Context.GetStatus() == EMovieScenePlayerStatus::Playing && Context.HasJumped());
+		if (bMoveToNewTime)
+		{
+			MediaPlayer->SetRate(0.0f);
+			MediaPlayer->Seek(StartTime);
+			MediaPlayer->SetBlockOnTimeRange(TRange<FTimespan>::Empty());
 		}
 	}
 
@@ -235,20 +244,20 @@ struct FMediaSectionExecutionToken
 				{
 					if (!MediaPlayer->SetRate(1.0f))
 					{
-					// Failed to set needed rate. Keep things blocked, as this means the player will still be returning the old rate, we will get here repeatedly
-					// and each time trigger a seek. A potentially very SLOW method of approximating backwards playback, but better
-					// than nothing.
-					MediaPlayer->Seek(MediaTime);
+						// Failed to set needed rate. Keep things blocked, as this means the player will still be returning the old rate, we will get here repeatedly
+						// and each time trigger a seek. A potentially very SLOW method of approximating backwards playback, but better
+						// than nothing.
+						MediaPlayer->Seek(MediaTime);
 					}
 				}
 				else if (Context.GetDirection() == EPlayDirection::Backwards && CurrentPlayerRate > 0.0f)
 				{
 					if (!MediaPlayer->SetRate(-1.0f))
 					{
-					// Failed to set needed rate. Keep things blocked, as this means the player will still be returning the old rate, we will get here repeatedly
-					// and each time trigger a seek. A potentially very SLOW method of approximating backwards playback, but better
-					// than nothing.
-					MediaPlayer->Seek(MediaTime);
+						// Failed to set needed rate. Keep things blocked, as this means the player will still be returning the old rate, we will get here repeatedly
+						// and each time trigger a seek. A potentially very SLOW method of approximating backwards playback, but better
+						// than nothing.
+						MediaPlayer->Seek(MediaTime);
 					}
 				}
 			}
@@ -434,7 +443,7 @@ void FMovieSceneMediaSectionTemplate::Initialize(const FMovieSceneEvaluationOper
 		if (IsEvaluating)
 		{
 			#if MOVIESCENEMEDIATEMPLATE_TRACE_EVALUATION
-				GLog->Logf(ELogVerbosity::Log, TEXT("Setting media player %p on media sound component %p"), MediaPlayer, Params.MediaSoundComponent);
+				GLog->Logf(ELogVerbosity::Log, TEXT("Setting media player %p on media sound component %p"), MediaPlayer, Params.MediaSoundComponent.Get());
 			#endif
 
 			Params.MediaSoundComponent->SetMediaPlayer(MediaPlayer);
@@ -442,7 +451,7 @@ void FMovieSceneMediaSectionTemplate::Initialize(const FMovieSceneEvaluationOper
 		else if (Params.MediaSoundComponent->GetMediaPlayer() == MediaPlayer)
 		{
 			#if MOVIESCENEMEDIATEMPLATE_TRACE_EVALUATION
-				GLog->Logf(ELogVerbosity::Log, TEXT("Resetting media player on media sound component %p"), Params.MediaSoundComponent);
+				GLog->Logf(ELogVerbosity::Log, TEXT("Resetting media player on media sound component %p"), Params.MediaSoundComponent.Get());
 			#endif
 
 			Params.MediaSoundComponent->SetMediaPlayer(nullptr);
@@ -454,7 +463,7 @@ void FMovieSceneMediaSectionTemplate::Initialize(const FMovieSceneEvaluationOper
 		if (IsEvaluating)
 		{
 			#if MOVIESCENEMEDIATEMPLATE_TRACE_EVALUATION
-				GLog->Logf(ELogVerbosity::Log, TEXT("Setting media player %p on media texture %p"), MediaPlayer, Params.MediaTexture);
+				GLog->Logf(ELogVerbosity::Log, TEXT("Setting media player %p on media texture %p"), MediaPlayer, Params.MediaTexture.Get());
 			#endif
 
 			Params.MediaTexture->SetMediaPlayer(MediaPlayer);
@@ -462,7 +471,7 @@ void FMovieSceneMediaSectionTemplate::Initialize(const FMovieSceneEvaluationOper
 		else if (Params.MediaTexture->GetMediaPlayer() == MediaPlayer)
 		{
 			#if MOVIESCENEMEDIATEMPLATE_TRACE_EVALUATION
-				GLog->Logf(ELogVerbosity::Log, TEXT("Resetting media player on media texture %p"), Params.MediaTexture);
+				GLog->Logf(ELogVerbosity::Log, TEXT("Resetting media player on media texture %p"), Params.MediaTexture.Get());
 			#endif
 
 			Params.MediaTexture->SetMediaPlayer(nullptr);
