Media: removed get/set of cache options intended for ImgMedia from the MediaSource and moved them into the new option map of FMediaPlayerOptions. Otherwise the cache options as used for playback in sequencer would become set generally on the media source.

#rnx

diff --git a/Engine/Plugins/Media/ImgMedia/Source/ImgMedia/Internal/Player/ImgMediaPlayer.h b/Engine/Plugins/Media/ImgMedia/Source/ImgMedia/Internal/Player/ImgMediaPlayer.h
--- a/Engine/Plugins/Media/ImgMedia/Source/ImgMedia/Internal/Player/ImgMediaPlayer.h
+++ b/Engine/Plugins/Media/ImgMedia/Source/ImgMedia/Internal/Player/ImgMediaPlayer.h
@@ -65,6 +65,7 @@ public:
 	virtual FString GetUrl() const override;
 	virtual IMediaView& GetView() override;
 	virtual bool Open(const FString& Url, const IMediaOptions* Options) override;
+	virtual bool Open(const FString& Url, const IMediaOptions* Options, const FMediaPlayerOptions* PlayerOptions) override;
 	virtual bool Open(const TSharedRef<FArchive, ESPMode::ThreadSafe>& Archive, const FString& OriginalUrl, const IMediaOptions* Options) override;
 	virtual void TickInput(FTimespan DeltaTime, FTimespan Timecode) override;
 	virtual bool FlushOnSeekStarted() const;
diff --git a/Engine/Plugins/Media/ImgMedia/Source/ImgMedia/Private/ImgMediaPrivate.h b/Engine/Plugins/Media/ImgMedia/Source/ImgMedia/Private/ImgMediaPrivate.h
--- a/Engine/Plugins/Media/ImgMedia/Source/ImgMedia/Private/ImgMediaPrivate.h
+++ b/Engine/Plugins/Media/ImgMedia/Source/ImgMedia/Private/ImgMediaPrivate.h
@@ -54,10 +54,4 @@ namespace ImgMedia
 
 	/** Name of the NumTilesY media option. */
 	static const FName NumTilesYOption("ImgNumTilesY");
-
-	/** Name of the smart cache media option. */
-	static const FName SmartCacheEnabled("ImgMediaSmartCacheEnabled");
-
-	/** Name of the smart cache time to look ahead media option. */
-	static const FName SmartCacheTimeToLookAhead("ImgMediaSmartCacheTimeToLookAhead");
 }
diff --git a/Engine/Plugins/Media/ImgMedia/Source/ImgMedia/Private/Player/ImgMediaPlayer.cpp b/Engine/Plugins/Media/ImgMedia/Source/ImgMedia/Private/Player/ImgMediaPlayer.cpp
--- a/Engine/Plugins/Media/ImgMedia/Source/ImgMedia/Private/Player/ImgMediaPlayer.cpp
+++ b/Engine/Plugins/Media/ImgMedia/Source/ImgMedia/Private/Player/ImgMediaPlayer.cpp
@@ -194,6 +194,14 @@ IMediaView& FImgMediaPlayer::GetView()
 
 
 bool FImgMediaPlayer::Open(const FString& Url, const IMediaOptions* Options)
+{
+	UE_LOG(LogImgMedia, Log, TEXT("Open() called without player options, smart caching will not be available."));
+	FMediaPlayerOptions NoOptions;
+	NoOptions.SetAllAsOptional();
+	return Open(Url, Options, &NoOptions);
+}
+
+bool FImgMediaPlayer::Open(const FString& Url, const IMediaOptions* Options, const FMediaPlayerOptions* PlayerOptions)
 {
 	Close();
 
@@ -218,19 +226,28 @@ bool FImgMediaPlayer::Open(const FString& Url, const IMediaOptions* Options)
 		Proxy = GetDefault<UImgMediaSettings>()->GetDefaultProxy();
 	}
 
+	auto GetBoolOption = [](const FMediaPlayerOptions* InPlayerOptions, const FName& InOptionName, bool bInDefaultValue) -> bool
+	{
+		const FVariant* Var = InPlayerOptions ? InPlayerOptions->InternalCustomOptions.Find(InOptionName) : nullptr;
+		return Var && Var->GetType() == EVariantTypes::Bool ? Var->GetValue<bool>() : bInDefaultValue;
+	};
+	auto GetDoubleOption = [](const FMediaPlayerOptions* InPlayerOptions, const FName& InOptionName, double InDefaultValue) -> double
+	{
+		const FVariant* Var = InPlayerOptions ? InPlayerOptions->InternalCustomOptions.Find(InOptionName) : nullptr;
+		return Var && Var->GetType() == EVariantTypes::Float ? (double)Var->GetValue<float>() : Var && Var->GetType() == EVariantTypes::Double ? Var->GetValue<double>() : InDefaultValue;
+	};
+
 	// get frame rate override, if any
 	FFrameRate FrameRateOverride(0, 0);
 	TSharedPtr<FImgMediaMipMapInfo, ESPMode::ThreadSafe> MipMapInfo;
 	bool bFillGapsInSequence = true;
-	bool bIsSmartCacheEnabled = false;
-	float SmartCacheTimeToLookAhead = 0.0f;
+	bool bIsSmartCacheEnabled = GetBoolOption(PlayerOptions, MediaPlayerOptionValues::ImgMediaSmartCacheEnabled(), false);
+	double SmartCacheTimeToLookAhead = GetDoubleOption(PlayerOptions, MediaPlayerOptionValues::ImgMediaSmartCacheTimeToLookAhead(), 0.0);
 	if (Options != nullptr)
 	{
 		FrameRateOverride.Denominator = Options->GetMediaOption(ImgMedia::FrameRateOverrideDenonimatorOption, 0LL);
 		FrameRateOverride.Numerator = Options->GetMediaOption(ImgMedia::FrameRateOverrideNumeratorOption, 0LL);
 		bFillGapsInSequence = Options->GetMediaOption(ImgMedia::FillGapsInSequenceOption, true);
-		bIsSmartCacheEnabled = Options->GetMediaOption(ImgMedia::SmartCacheEnabled, false);
-		SmartCacheTimeToLookAhead = Options->GetMediaOption(ImgMedia::SmartCacheTimeToLookAhead, 0.0f);
 		TSharedPtr<IMediaOptions::FDataContainer, ESPMode::ThreadSafe> DefaultValue;
 		TSharedPtr<IMediaOptions::FDataContainer, ESPMode::ThreadSafe> DataContainer = Options->GetMediaOption(ImgMedia::MipMapInfoOption, DefaultValue);
 		if (DataContainer.IsValid())
@@ -266,7 +283,7 @@ bool FImgMediaPlayer::Open(const FString& Url, const IMediaOptions* Options)
 	}
 
 	// initialize image loader on a separate thread
-	FImgMediaLoaderSmartCacheSettings SmartCacheSettings(bIsSmartCacheEnabled, SmartCacheTimeToLookAhead);
+	FImgMediaLoaderSmartCacheSettings SmartCacheSettings(bIsSmartCacheEnabled, (float)SmartCacheTimeToLookAhead);
 	Loader = MakeShared<FImgMediaLoader, ESPMode::ThreadSafe>(Scheduler.ToSharedRef(),
 		GlobalCache.ToSharedRef(), MipMapInfo, bFillGapsInSequence, SmartCacheSettings);
 	Scheduler->RegisterLoader(Loader.ToSharedRef());
diff --git a/Engine/Plugins/Media/MediaCompositing/Source/MediaCompositing/Private/MovieScene/MovieSceneMediaTemplate.cpp b/Engine/Plugins/Media/MediaCompositing/Source/MediaCompositing/Private/MovieScene/MovieSceneMediaTemplate.cpp
--- a/Engine/Plugins/Media/MediaCompositing/Source/MediaCompositing/Private/MovieScene/MovieSceneMediaTemplate.cpp
+++ b/Engine/Plugins/Media/MediaCompositing/Source/MediaCompositing/Private/MovieScene/MovieSceneMediaTemplate.cpp
@@ -9,6 +9,7 @@
 #include "MediaPlayerProxyInterface.h"
 #include "MediaSoundComponent.h"
 #include "MediaSource.h"
+#include "MediaSourceOptions.h"
 #include "MediaTexture.h"
 #include "UObject/Package.h"
 #include "UObject/GCObject.h"
@@ -26,14 +27,45 @@
 /* Local helpers
  *****************************************************************************/
 
+struct FMediaSectionBaseMediaSourceOptions
+{
+	FMediaSourceCacheSettings CacheSettings;
+};
+
+namespace FMediaSectionBaseMediaUtils
+{
+	FMediaSourceCacheSettings GetCurrentCacheSettingsFromPlayer(const UMediaPlayer* InPlayer)
+	{
+		FMediaSourceCacheSettings cs;
+		if (InPlayer)
+		{
+			TOptional<FMediaPlayerOptions> ActivePlayerOptions = InPlayer->GetPlayerFacade()->ActivePlayerOptions;
+			if (ActivePlayerOptions.IsSet() 
+				&& ActivePlayerOptions.GetValue().InternalCustomOptions.Contains(MediaPlayerOptionValues::ImgMediaSmartCacheEnabled())
+				&& ActivePlayerOptions.GetValue().InternalCustomOptions.Contains(MediaPlayerOptionValues::ImgMediaSmartCacheTimeToLookAhead()))
+			{
+				FVariant vEnabled = ActivePlayerOptions.GetValue().InternalCustomOptions[MediaPlayerOptionValues::ImgMediaSmartCacheEnabled()];
+				FVariant vTTLA = ActivePlayerOptions.GetValue().InternalCustomOptions[MediaPlayerOptionValues::ImgMediaSmartCacheTimeToLookAhead()];
+				if (vEnabled.GetType() == EVariantTypes::Bool && (vTTLA.GetType() == EVariantTypes::Float || vTTLA.GetType() == EVariantTypes::Double))
+				{
+					cs.bOverride = vEnabled.GetValue<bool>();
+					cs.TimeToLookAhead = vTTLA.GetType() == EVariantTypes::Float ? (double)vTTLA.GetValue<float>() : vTTLA.GetValue<double>();
+				}
+			}
+		}
+		return cs;
+	}
+}
+
 /** Base struct for exectution tokens. */
 struct FMediaSectionBaseExecutionToken
 	: IMovieSceneExecutionToken
 {
-	FMediaSectionBaseExecutionToken(UMediaSource* InMediaSource, const FMovieSceneObjectBindingID& InMediaSourceProxy, int32 InMediaSourceProxyIndex)
+	FMediaSectionBaseExecutionToken(UMediaSource* InMediaSource, const FMediaSectionBaseMediaSourceOptions& InMediaSourceOptions, const FMovieSceneObjectBindingID& InMediaSourceProxy, int32 InMediaSourceProxyIndex)
 		: BaseMediaSource(InMediaSource)
 		, MediaSourceProxy(InMediaSourceProxy)
 		, MediaSourceProxyIndex(InMediaSourceProxyIndex)
+		, BaseMediaSourceOptions(InMediaSourceOptions)
 	{
 	}
 
@@ -56,17 +88,23 @@ struct FMediaSectionBaseExecutionToken
 	 */
 	bool IsMediaSourceProxyValid() const { return MediaSourceProxy.IsValid(); }
 
+	/**
+	 * Gets the media source options
+	 */
+	const FMediaSectionBaseMediaSourceOptions& GetBaseMediaSourceOptions() const { return BaseMediaSourceOptions; }
+
 private:
 	UMediaSource* BaseMediaSource;
 	FMovieSceneObjectBindingID MediaSourceProxy;
 	int32 MediaSourceProxyIndex = 0;
+	FMediaSectionBaseMediaSourceOptions BaseMediaSourceOptions;
 };
 
 struct FMediaSectionPreRollExecutionToken
 	: FMediaSectionBaseExecutionToken
 {
-	FMediaSectionPreRollExecutionToken(UMediaSource* InMediaSource, FMovieSceneObjectBindingID InMediaSourceProxy, int32 InMediaSourceProxyIndex, FTimespan InStartTimeSeconds)
-		: FMediaSectionBaseExecutionToken(InMediaSource, InMediaSourceProxy, InMediaSourceProxyIndex)
+	FMediaSectionPreRollExecutionToken(UMediaSource* InMediaSource, const FMediaSectionBaseMediaSourceOptions& InMediaSourceOptions, FMovieSceneObjectBindingID InMediaSourceProxy, int32 InMediaSourceProxyIndex, FTimespan InStartTimeSeconds)
+		: FMediaSectionBaseExecutionToken(InMediaSource, InMediaSourceOptions, InMediaSourceProxy, InMediaSourceProxyIndex)
 		, StartTime(InStartTimeSeconds)
 	{ }
 
@@ -87,16 +125,19 @@ struct FMediaSectionPreRollExecutionToken
 		// open the media source if necessary
 		if (MediaPlayer->GetUrl().IsEmpty())
 		{
+			SectionData.SeekOnOpen(StartTime);
+
 			FMediaPlayerOptions Options;
 			Options.SetAllAsOptional();
-
+			Options.InternalCustomOptions.Emplace(MediaPlayerOptionValues::Environment(), MediaPlayerOptionValues::Environment_Sequencer());
 			if (PlayerProxyInterface != nullptr)
 			{
-				MediaSource->SetCacheSettings(PlayerProxyInterface->GetCacheSettings());
+				// Set cache settings.
+				const FMediaSourceCacheSettings& CacheSettings = PlayerProxyInterface->GetCacheSettings();
+				Options.InternalCustomOptions.Emplace(MediaPlayerOptionValues::ImgMediaSmartCacheEnabled(), FVariant(CacheSettings.bOverride));
+				Options.InternalCustomOptions.Emplace(MediaPlayerOptionValues::ImgMediaSmartCacheTimeToLookAhead(), FVariant(CacheSettings.TimeToLookAhead));
 			}
-			SectionData.SeekOnOpen(StartTime);
 
-			Options.InternalCustomOptions.Emplace(MediaPlayerOptionValues::Environment(), MediaPlayerOptionValues::Environment_Sequencer());
 			MediaPlayer->OpenSourceWithOptions(MediaSource, Options);
 			return;
 		}
@@ -119,8 +160,8 @@ private:
 struct FMediaSectionExecutionToken
 	: FMediaSectionBaseExecutionToken
 {
-	FMediaSectionExecutionToken(UMediaSource* InMediaSource, FMovieSceneObjectBindingID InMediaSourceProxy, int32 InMediaSourceProxyIndex, float InProxyTextureBlend, bool bInCanPlayerBeOpen, FTimespan InCurrentTime, FTimespan InFrameDuration)
-		: FMediaSectionBaseExecutionToken(InMediaSource, InMediaSourceProxy, InMediaSourceProxyIndex)
+	FMediaSectionExecutionToken(UMediaSource* InMediaSource, const FMediaSectionBaseMediaSourceOptions& InMediaSourceOptions, FMovieSceneObjectBindingID InMediaSourceProxy, int32 InMediaSourceProxyIndex, float InProxyTextureBlend, bool bInCanPlayerBeOpen, FTimespan InCurrentTime, FTimespan InFrameDuration)
+		: FMediaSectionBaseExecutionToken(InMediaSource, InMediaSourceOptions, InMediaSourceProxy, InMediaSourceProxyIndex)
 		, CurrentTime(InCurrentTime)
 		, FrameDuration(InFrameDuration)
 		, PlaybackRate(1.0f)
@@ -170,26 +211,27 @@ struct FMediaSectionExecutionToken
 			return;
 		}
 
-		bool bCacheSettingsChanged = false;
-		FMediaSourceCacheSettings CurrentCacheSettings;
-		if (PlayerProxyInterface != nullptr && MediaSource->GetCacheSettings(CurrentCacheSettings))
-		{
-			bCacheSettingsChanged = (CurrentCacheSettings != PlayerProxyInterface->GetCacheSettings());
-		}
+
+		// Check if the cache options have changed mid playback.
+		//const FMediaSourceCacheSettings CurrentCacheSettings = GetBaseMediaSourceOptions().CacheSettings;
+		const FMediaSourceCacheSettings CurrentCacheSettings = FMediaSectionBaseMediaUtils::GetCurrentCacheSettingsFromPlayer(MediaPlayer);
+		bool bCacheSettingsChanged = PlayerProxyInterface ? CurrentCacheSettings != PlayerProxyInterface->GetCacheSettings() : false;
 
 		// open the media source if necessary
 		if (MediaPlayer->GetUrl().IsEmpty() || bCacheSettingsChanged)
 		{
+			SectionData.SeekOnOpen(CurrentTime);
+
 			FMediaPlayerOptions Options;
 			Options.SetAllAsOptional();
-
+			Options.InternalCustomOptions.Emplace(MediaPlayerOptionValues::Environment(), MediaPlayerOptionValues::Environment_Sequencer());
 			if (PlayerProxyInterface != nullptr)
 			{
-				MediaSource->SetCacheSettings(PlayerProxyInterface->GetCacheSettings());
+				const FMediaSourceCacheSettings& CacheSettings = PlayerProxyInterface->GetCacheSettings();
+				Options.InternalCustomOptions.Emplace(MediaPlayerOptionValues::ImgMediaSmartCacheEnabled(), FVariant(CacheSettings.bOverride));
+				Options.InternalCustomOptions.Emplace(MediaPlayerOptionValues::ImgMediaSmartCacheTimeToLookAhead(), FVariant(CacheSettings.TimeToLookAhead));
 			}
-			SectionData.SeekOnOpen(CurrentTime);
 
-			Options.InternalCustomOptions.Emplace(MediaPlayerOptionValues::Environment(), MediaPlayerOptionValues::Environment_Sequencer());
 			// Setup an initial blocking range - MediaFramework will block (even through the opening process) in its next tick...
 			MediaPlayer->SetBlockOnTimeRange(TRange<FTimespan>(CurrentTime, CurrentTime + FrameDuration));
 			MediaPlayer->OpenSourceWithOptions(MediaSource, Options);
@@ -307,10 +349,7 @@ FMovieSceneMediaSectionTemplate::FMovieSceneMediaSectionTemplate(const UMovieSce
 	Params.MediaSoundComponent = InSection.MediaSoundComponent;
 	Params.bLooping = InSection.bLooping;
 	Params.StartFrameOffset = InSection.StartFrameOffset;
-	if (Params.MediaSource != nullptr)
-	{
-		Params.MediaSource->SetCacheSettings(InSection.CacheSettings);
-	}
+	Params.CacheSettings = InSection.CacheSettings;
 
 	// If using an external media player link it here so we don't automatically create it later.
 	Params.MediaPlayer = InSection.bUseExternalMediaPlayer ? InSection.ExternalMediaPlayer : nullptr;
@@ -348,7 +387,9 @@ void FMovieSceneMediaSectionTemplate::Evaluate(const FMovieSceneEvaluationOperan
 		const double StartFrameInSeconds = FrameRate.AsSeconds(StartFrame);
 		const int64 StartTicks = static_cast<int64>(StartFrameInSeconds * ETimespan::TicksPerSecond);
 
-		ExecutionTokens.Add(FMediaSectionPreRollExecutionToken(MediaSource, Params.MediaSourceProxy, Params.MediaSourceProxyIndex, FTimespan(StartTicks)));
+		FMediaSectionBaseMediaSourceOptions Options;
+		Options.CacheSettings = Params.CacheSettings;
+		ExecutionTokens.Add(FMediaSectionPreRollExecutionToken(MediaSource, Options, Params.MediaSourceProxy, Params.MediaSourceProxyIndex, FTimespan(StartTicks)));
 	}
 	else if (!Context.IsPostRoll() && (Context.GetTime().FrameNumber < Params.SectionEndFrame))
 	{
@@ -377,7 +418,9 @@ void FMovieSceneMediaSectionTemplate::Evaluate(const FMovieSceneEvaluationOperan
 			);
 		#endif
 
-		ExecutionTokens.Add(FMediaSectionExecutionToken(MediaSource, Params.MediaSourceProxy, Params.MediaSourceProxyIndex, ProxyTextureBlend, bCanPlayerBeOpen, FTimespan(FrameTicks), FTimespan(FrameDurationTicks)));
+		FMediaSectionBaseMediaSourceOptions Options;
+		Options.CacheSettings = Params.CacheSettings;
+		ExecutionTokens.Add(FMediaSectionExecutionToken(MediaSource, Options, Params.MediaSourceProxy, Params.MediaSourceProxyIndex, ProxyTextureBlend, bCanPlayerBeOpen, FTimespan(FrameTicks), FTimespan(FrameDurationTicks)));
 	}
 }
 
diff --git a/Engine/Plugins/Media/MediaCompositing/Source/MediaCompositing/Private/MovieScene/MovieSceneMediaTemplate.h b/Engine/Plugins/Media/MediaCompositing/Source/MediaCompositing/Private/MovieScene/MovieSceneMediaTemplate.h
--- a/Engine/Plugins/Media/MediaCompositing/Source/MediaCompositing/Private/MovieScene/MovieSceneMediaTemplate.h
+++ b/Engine/Plugins/Media/MediaCompositing/Source/MediaCompositing/Private/MovieScene/MovieSceneMediaTemplate.h
@@ -6,6 +6,7 @@
 #include "Evaluation/MovieSceneEvalTemplate.h"
 #include "Evaluation/MovieScenePropertyTemplate.h"
 #include "MovieSceneObjectBindingID.h"
+#include "MediaSourceOptions.h"
 
 #include "MovieSceneMediaTemplate.generated.h"
 
@@ -55,6 +56,9 @@ struct FMovieSceneMediaSectionParams
 	UPROPERTY()
 	FMovieSceneFloatChannel ProxyTextureBlend;
 
+	UPROPERTY()
+	FMediaSourceCacheSettings CacheSettings;
+
 	FMovieSceneMediaSectionParams()
 		: MediaSoundComponent(nullptr)
 		, MediaSource(nullptr)
diff --git a/Engine/Plugins/Media/MediaPlate/Source/MediaPlate/Private/MediaPlateComponent.cpp b/Engine/Plugins/Media/MediaPlate/Source/MediaPlate/Private/MediaPlateComponent.cpp
--- a/Engine/Plugins/Media/MediaPlate/Source/MediaPlate/Private/MediaPlateComponent.cpp
+++ b/Engine/Plugins/Media/MediaPlate/Source/MediaPlate/Private/MediaPlateComponent.cpp
@@ -564,9 +564,6 @@ bool UMediaPlateComponent::PlayMediaSource(UMediaSource* InMediaSource, bool bIn
 
 	if (InMediaSource != nullptr)
 	{
-		// Set cache settings.
-		InMediaSource->SetCacheSettings(CacheSettings);
-
 		// Set media options.
 		if (MediaPlayer != nullptr)
 		{
@@ -578,8 +575,11 @@ bool UMediaPlateComponent::PlayMediaSource(UMediaSource* InMediaSource, bool bIn
 			Options.PlayOnOpen = bInPlayOnOpen ? EMediaPlayerOptionBooleanOverride::Enabled : EMediaPlayerOptionBooleanOverride::Disabled;
 			Options.Loop = (bLoop && (bIsPlaylist == false)) ? EMediaPlayerOptionBooleanOverride::Enabled : EMediaPlayerOptionBooleanOverride::Disabled;
 			Options.InternalCustomOptions.Emplace(MediaPlayerOptionValues::Environment(), MediaPlayerOptionValues::Environment_Preview());
-			bIsPlaying = MediaPlayer->OpenSourceWithOptions(InMediaSource, Options);
+			// Set cache settings.
+			Options.InternalCustomOptions.Emplace(MediaPlayerOptionValues::ImgMediaSmartCacheEnabled(), FVariant(CacheSettings.bOverride));
+			Options.InternalCustomOptions.Emplace(MediaPlayerOptionValues::ImgMediaSmartCacheTimeToLookAhead(), FVariant(CacheSettings.TimeToLookAhead));
 
+			bIsPlaying = MediaPlayer->OpenSourceWithOptions(InMediaSource, Options);
 			// Did we play anything?
 			if (bIsPlaying)
 			{
diff --git a/Engine/Plugins/Media/MediaPlayerEditor/Source/MediaPlayerEditor/Private/Shared/MediaSourceRenderer.cpp b/Engine/Plugins/Media/MediaPlayerEditor/Source/MediaPlayerEditor/Private/Shared/MediaSourceRenderer.cpp
--- a/Engine/Plugins/Media/MediaPlayerEditor/Source/MediaPlayerEditor/Private/Shared/MediaSourceRenderer.cpp
+++ b/Engine/Plugins/Media/MediaPlayerEditor/Source/MediaPlayerEditor/Private/Shared/MediaSourceRenderer.cpp
@@ -33,23 +33,19 @@ UMediaTexture* UMediaSourceRenderer::Open(UMediaSource* InMediaSource)
 		MediaTexture->UpdateResource();
 		MediaSource = InMediaSource;
 
-		// For image media, we avoid filling the global cache which will needlessly hold onto frame data.
-		FMediaSourceCacheSettings OriginalCacheSettings;
-		MediaSource->GetCacheSettings(OriginalCacheSettings);
-		MediaSource->SetCacheSettings(FMediaSourceCacheSettings{ true, 0.2f });
-
 		// Start playing the media.
 		bIsSeekActive = false;
 		FMediaPlayerOptions Options;
 		Options.PlayOnOpen = EMediaPlayerOptionBooleanOverride::Enabled;
 		Options.Loop = EMediaPlayerOptionBooleanOverride::Disabled;
+		// For image media, we avoid filling the global cache which will needlessly hold onto frame data.
+		Options.InternalCustomOptions.Emplace(MediaPlayerOptionValues::ImgMediaSmartCacheEnabled(), FVariant(true));
+		Options.InternalCustomOptions.Emplace(MediaPlayerOptionValues::ImgMediaSmartCacheTimeToLookAhead(), FVariant(0.2f));
 		bool bIsPlaying = MediaPlayer->OpenSourceWithOptions(MediaSource, Options);
 		if (bIsPlaying == false)
 		{
 			Close();
 		}
-
-		MediaSource->SetCacheSettings(OriginalCacheSettings);
 	}
 
 	return MediaTexture.Get();
diff --git a/Engine/Source/Runtime/MediaAssets/Private/Assets/MediaPlayer.cpp b/Engine/Source/Runtime/MediaAssets/Private/Assets/MediaPlayer.cpp
--- a/Engine/Source/Runtime/MediaAssets/Private/Assets/MediaPlayer.cpp
+++ b/Engine/Source/Runtime/MediaAssets/Private/Assets/MediaPlayer.cpp
@@ -532,7 +532,6 @@ bool UMediaPlayer::IsClosed() const
 
 bool UMediaPlayer::IsReady() const
 {
-	UE_LOG(LogMediaAssets, VeryVerbose, TEXT("%s.IsReady"), *GetFName().ToString());
 	return PlayerFacade->IsReady();
 }
 
diff --git a/Engine/Source/Runtime/MediaAssets/Private/Assets/MediaSource.cpp b/Engine/Source/Runtime/MediaAssets/Private/Assets/MediaSource.cpp
--- a/Engine/Source/Runtime/MediaAssets/Private/Assets/MediaSource.cpp
+++ b/Engine/Source/Runtime/MediaAssets/Private/Assets/MediaSource.cpp
@@ -16,32 +16,6 @@
 
 #include UE_INLINE_GENERATED_CPP_BY_NAME(MediaSource)
 
-static const FLazyName ImgMediaSmartCacheEnabledName(TEXT("ImgMediaSmartCacheEnabled"));
-static const FLazyName ImgMediaSmartCacheTimeToLookAheadName(TEXT("ImgMediaSmartCacheTimeToLookAhead"));
-
-void UMediaSource::SetCacheSettings(const FMediaSourceCacheSettings& Settings)
-{
-	SetMediaOptionBool(ImgMediaSmartCacheEnabledName, Settings.bOverride);
-	SetMediaOptionFloat(ImgMediaSmartCacheTimeToLookAheadName, Settings.TimeToLookAhead);
-}
-
-bool UMediaSource::GetCacheSettings(FMediaSourceCacheSettings& OutSettings) const
-{
-	if (!HasMediaOption(ImgMediaSmartCacheEnabledName))
-	{
-		return false;
-	}
-
-	if (!HasMediaOption(ImgMediaSmartCacheTimeToLookAheadName))
-	{
-		return false;
-	}
-
-	OutSettings.bOverride = GetMediaOption(ImgMediaSmartCacheEnabledName, false);
-	OutSettings.TimeToLookAhead = GetMediaOption(ImgMediaSmartCacheTimeToLookAheadName, 0.0f);
-	return true;
-}
-
 #if WITH_EDITOR
 
 void UMediaSource::GenerateThumbnail()
diff --git a/Engine/Source/Runtime/MediaAssets/Private/Assets/MediaSourceOptions.cpp b/Engine/Source/Runtime/MediaAssets/Private/Assets/MediaSourceOptions.cpp
new file mode 100644
--- /dev/null
+++ b/Engine/Source/Runtime/MediaAssets/Private/Assets/MediaSourceOptions.cpp
@@ -0,0 +1,6 @@
+// Copyright Epic Games, Inc. All Rights Reserved.
+
+#include "MediaSourceOptions.h"
+
+#include UE_INLINE_GENERATED_CPP_BY_NAME(MediaSourceOptions)
+
diff --git a/Engine/Source/Runtime/MediaAssets/Public/MediaSource.h b/Engine/Source/Runtime/MediaAssets/Public/MediaSource.h
--- a/Engine/Source/Runtime/MediaAssets/Public/MediaSource.h
+++ b/Engine/Source/Runtime/MediaAssets/Public/MediaSource.h
@@ -7,7 +7,6 @@
 #include "CoreTypes.h"
 #include "Delegates/Delegate.h"
 #include "IMediaOptions.h"
-#include "IMediaOptions.h"
 #include "Internationalization/Text.h"
 #include "Misc/AssertionMacros.h"
 #include "Misc/Variant.h"
@@ -17,6 +16,7 @@
 #include "UObject/ObjectMacros.h"
 #include "UObject/ScriptMacros.h"
 #include "UObject/UObjectGlobals.h"
+#include "MediaSourceOptions.h"
 
 #include "MediaSource.generated.h"
 
@@ -29,39 +29,6 @@ struct FFrame;
 /** Delegate for creating a media source from a string. */
 DECLARE_DELEGATE_RetVal_TwoParams(UMediaSource*, FMediaSourceSpawnDelegate, const FString&, UObject*);
 
-/** Cache settings to pass to the player. */
-USTRUCT(BlueprintType)
-struct FMediaSourceCacheSettings
-{
-	GENERATED_USTRUCT_BODY()
-
-	/**
-	 * Override the default cache settings.
-	 * Currently only the ImgMedia player supports these settings.
-	 */
-	UPROPERTY(EditAnywhere, Category = "Media Cache")
-	bool bOverride = false;
-
-	/**
-	 * The cache will fill up with frames that are up to this time from the current time.
-	 * E.g. if this is 0.2, and we are at time index 5 seconds,
-	 * then we will fill the cache with frames between 5 seconds and 5.2 seconds.
-	 */
-	UPROPERTY(EditAnywhere, Category = "Media Cache")
-	float TimeToLookAhead = 0.2f;
-
-
-	inline bool operator==(const FMediaSourceCacheSettings& Other) const
-	{
-		return (Other.bOverride == bOverride) && FMath::IsNearlyEqual(Other.TimeToLookAhead, TimeToLookAhead);
-	}
-
-	inline bool operator!=(const FMediaSourceCacheSettings& Other) const
-	{
-		return !(*this == Other);
-	}
-};
-
 /**
  * Abstract base class for media sources.
  *
@@ -97,19 +64,6 @@ public:
 	UFUNCTION(BlueprintCallable, Category="Media|MediaSource")
 	MEDIAASSETS_API virtual bool Validate() const PURE_VIRTUAL(UMediaSource::Validate, return false;);
 
-	/**
-	 * Call this to set cache settings to pass to the player.
-	 */
-	MEDIAASSETS_API void SetCacheSettings(const FMediaSourceCacheSettings& Settings);
-
-	/**
-	 * Get the media source cache settings, if present.
-	 * 
-	 * @param OutSettings Cache settings
-	 * @return true if the cache settings are present, false otherwise.
-	 */
-	MEDIAASSETS_API bool GetCacheSettings(FMediaSourceCacheSettings& OutSettings) const;
-
 #if WITH_EDITOR
 
 	/**
diff --git a/Engine/Source/Runtime/MediaAssets/Public/MediaSourceOptions.h b/Engine/Source/Runtime/MediaAssets/Public/MediaSourceOptions.h
new file mode 100644
--- /dev/null
+++ b/Engine/Source/Runtime/MediaAssets/Public/MediaSourceOptions.h
@@ -0,0 +1,50 @@
+// Copyright Epic Games, Inc. All Rights Reserved.
+
+#pragma once
+
+#include "Containers/Map.h"
+#include "Containers/UnrealString.h"
+#include "CoreTypes.h"
+#include "Internationalization/Text.h"
+#include "Misc/AssertionMacros.h"
+#include "Templates/SharedPointer.h"
+#include "UObject/NameTypes.h"
+#include "UObject/Object.h"
+#include "UObject/ObjectMacros.h"
+#include "UObject/ScriptMacros.h"
+#include "UObject/UObjectGlobals.h"
+
+#include "MediaSourceOptions.generated.h"
+
+/** Cache settings to pass to the player. */
+USTRUCT(BlueprintType)
+struct FMediaSourceCacheSettings
+{
+	GENERATED_USTRUCT_BODY()
+
+	/**
+	 * Override the default cache settings.
+	 * Currently only the ImgMedia player supports these settings.
+	 */
+	UPROPERTY(EditAnywhere, Category = "Media Cache")
+	bool bOverride = false;
+
+	/**
+	 * The cache will fill up with frames that are up to this time from the current time.
+	 * E.g. if this is 0.2, and we are at time index 5 seconds,
+	 * then we will fill the cache with frames between 5 seconds and 5.2 seconds.
+	 */
+	UPROPERTY(EditAnywhere, Category = "Media Cache")
+	float TimeToLookAhead = 0.2f;
+
+
+	inline bool operator==(const FMediaSourceCacheSettings& Other) const
+	{
+		return (Other.bOverride == bOverride) && FMath::IsNearlyEqual(Other.TimeToLookAhead, TimeToLookAhead);
+	}
+
+	inline bool operator!=(const FMediaSourceCacheSettings& Other) const
+	{
+		return !(*this == Other);
+	}
+};
diff --git a/Engine/Source/Runtime/MediaUtils/Public/MediaPlayerOptions.h b/Engine/Source/Runtime/MediaUtils/Public/MediaPlayerOptions.h
--- a/Engine/Source/Runtime/MediaUtils/Public/MediaPlayerOptions.h
+++ b/Engine/Source/Runtime/MediaUtils/Public/MediaPlayerOptions.h
@@ -160,4 +160,9 @@ namespace MediaPlayerOptionValues
 	{ static FName OptName(TEXT("Sequencer")); return OptName; }
 	inline const FName& Environment_Preview() 
 	{ static FName OptName(TEXT("Preview")); return OptName; }
+
+	inline const FName& ImgMediaSmartCacheEnabled() 
+	{ static FName OptName(TEXT("ImgMediaSmartCacheEnabled")); return OptName; }
+	inline const FName& ImgMediaSmartCacheTimeToLookAhead() 
+	{ static FName OptName(TEXT("ImgMediaSmartCacheTimeToLookAhead")); return OptName; }
 }
