Grab bag of non-unity compile fixes

diff --git a/Engine/Plugins/Experimental/RenderGrid/Source/RenderGridEditor/Private/UI/SRenderGridJobList.cpp b/Engine/Plugins/Experimental/RenderGrid/Source/RenderGridEditor/Private/UI/SRenderGridJobList.cpp
--- a/Engine/Plugins/Experimental/RenderGrid/Source/RenderGridEditor/Private/UI/SRenderGridJobList.cpp
+++ b/Engine/Plugins/Experimental/RenderGrid/Source/RenderGridEditor/Private/UI/SRenderGridJobList.cpp
@@ -28,6 +28,7 @@
 #include "Widgets/Input/SCheckBox.h"
 #include "Widgets/Input/SSearchBox.h"
 #include "Widgets/Layout/SScaleBox.h"
+#include "LevelSequence.h"
 
 #define LOCTEXT_NAMESPACE "SRenderGridJobList"
 
diff --git a/Engine/Plugins/Media/MediaCompositing/Source/MediaCompositing/Private/MovieScene/MovieSceneMediaSection.cpp b/Engine/Plugins/Media/MediaCompositing/Source/MediaCompositing/Private/MovieScene/MovieSceneMediaSection.cpp
--- a/Engine/Plugins/Media/MediaCompositing/Source/MediaCompositing/Private/MovieScene/MovieSceneMediaSection.cpp
+++ b/Engine/Plugins/Media/MediaCompositing/Source/MediaCompositing/Private/MovieScene/MovieSceneMediaSection.cpp
@@ -7,6 +7,7 @@
 #include "MediaPlayerProxyInterface.h"
 #include "MovieScene.h"
 #include "Misc/FrameRate.h"
+#include "EntitySystem/MovieSceneSharedPlaybackState.h"
 
 #include UE_INLINE_GENERATED_CPP_BY_NAME(MovieSceneMediaSection)
 
diff --git a/Engine/Source/Editor/MovieSceneTools/Private/TrackEditors/CameraCutTrackEditor.cpp b/Engine/Source/Editor/MovieSceneTools/Private/TrackEditors/CameraCutTrackEditor.cpp
--- a/Engine/Source/Editor/MovieSceneTools/Private/TrackEditors/CameraCutTrackEditor.cpp
+++ b/Engine/Source/Editor/MovieSceneTools/Private/TrackEditors/CameraCutTrackEditor.cpp
@@ -14,6 +14,7 @@
 #include "TrackEditorThumbnail/TrackEditorThumbnailPool.h"
 #include "TrackInstances/MovieSceneCameraCutTrackInstance.h"
 #include "Tracks/MovieSceneCameraCutTrack.h"
+#include "Evaluation/MovieSceneEvaluationTemplateInstance.h"
 
 #include "ActorEditorUtils.h"
 #include "ActorTreeItem.h"
