[MediaCompositing] Updating comments to reflect current class hierarchy.

diff --git a/Engine/Plugins/Media/MediaCompositing/Source/MediaCompositing/Private/MovieScene/MovieSceneMediaTrack.cpp b/Engine/Plugins/Media/MediaCompositing/Source/MediaCompositing/Private/MovieScene/MovieSceneMediaTrack.cpp
--- a/Engine/Plugins/Media/MediaCompositing/Source/MediaCompositing/Private/MovieScene/MovieSceneMediaTrack.cpp
+++ b/Engine/Plugins/Media/MediaCompositing/Source/MediaCompositing/Private/MovieScene/MovieSceneMediaTrack.cpp
@@ -177,7 +177,7 @@ void UMovieSceneMediaTrack::TickOutput()
 }
 
 
-/* UMovieScenePropertyTrack interface
+/* UMovieSceneTrack interface
  *****************************************************************************/
 
 void UMovieSceneMediaTrack::AddSection(UMovieSceneSection& Section)
diff --git a/Engine/Plugins/Media/MediaCompositing/Source/MediaCompositing/Public/MovieSceneMediaTrack.h b/Engine/Plugins/Media/MediaCompositing/Source/MediaCompositing/Public/MovieSceneMediaTrack.h
--- a/Engine/Plugins/Media/MediaCompositing/Source/MediaCompositing/Public/MovieSceneMediaTrack.h
+++ b/Engine/Plugins/Media/MediaCompositing/Source/MediaCompositing/Public/MovieSceneMediaTrack.h
@@ -58,7 +58,7 @@ public:
 
 public:
 
-	//~ UMovieScenePropertyTrack interface
+	//~ UMovieSceneTrack interface
 
 	virtual void AddSection(UMovieSceneSection& Section) override;
 	virtual bool SupportsType(TSubclassOf<UMovieSceneSection> SectionClass) const override;
