[Media] Media Track Scrubbing improvements

Improve sequencer scrubbing by throttling seek requests to allow media players to deliver 1 frame of output before issuing a new seek.

diff --git a/Engine/Plugins/Media/MediaCompositing/Source/MediaCompositing/Private/MovieScene/MovieSceneMediaTemplate.cpp b/Engine/Plugins/Media/MediaCompositing/Source/MediaCompositing/Private/MovieScene/MovieSceneMediaTemplate.cpp
--- a/Engine/Plugins/Media/MediaCompositing/Source/MediaCompositing/Private/MovieScene/MovieSceneMediaTemplate.cpp
+++ b/Engine/Plugins/Media/MediaCompositing/Source/MediaCompositing/Private/MovieScene/MovieSceneMediaTemplate.cpp
@@ -44,8 +44,8 @@ namespace FMediaSectionBaseMediaUtils
 				&& ActivePlayerOptions.GetValue().InternalCustomOptions.Contains(MediaPlayerOptionValues::ImgMediaSmartCacheEnabled())
 				&& ActivePlayerOptions.GetValue().InternalCustomOptions.Contains(MediaPlayerOptionValues::ImgMediaSmartCacheTimeToLookAhead()))
 			{
-				FVariant vEnabled = ActivePlayerOptions.GetValue().InternalCustomOptions[MediaPlayerOptionValues::ImgMediaSmartCacheEnabled()];
-				FVariant vTTLA = ActivePlayerOptions.GetValue().InternalCustomOptions[MediaPlayerOptionValues::ImgMediaSmartCacheTimeToLookAhead()];
+				const FVariant vEnabled = ActivePlayerOptions.GetValue().InternalCustomOptions[MediaPlayerOptionValues::ImgMediaSmartCacheEnabled()];
+				const FVariant vTTLA = ActivePlayerOptions.GetValue().InternalCustomOptions[MediaPlayerOptionValues::ImgMediaSmartCacheTimeToLookAhead()];
 				if (vEnabled.GetType() == EVariantTypes::Bool && (vTTLA.GetType() == EVariantTypes::Float || vTTLA.GetType() == EVariantTypes::Double))
 				{
 					cs.bOverride = vEnabled.GetValue<bool>();
@@ -57,7 +57,7 @@ namespace FMediaSectionBaseMediaUtils
 	}
 }
 
-/** Base struct for exectution tokens. */
+/** Base struct for execution tokens. */
 struct FMediaSectionBaseExecutionToken
 	: IMovieSceneExecutionToken
 {
@@ -72,7 +72,7 @@ struct FMediaSectionBaseExecutionToken
 	/**
 	 * Gets the media source from either the proxy binding or the media source.
 	 */
-	UMediaSource* GetMediaSource(IMovieScenePlayer& Player, FMovieSceneSequenceID SequenceID)
+	UMediaSource* GetMediaSource(IMovieScenePlayer& Player, FMovieSceneSequenceID SequenceID) const
 	{
 		return UMovieSceneMediaSection::GetMediaSourceOrProxy(Player, SequenceID,
 			BaseMediaSource, MediaSourceProxy, MediaSourceProxyIndex);
@@ -114,7 +114,7 @@ struct FMediaSectionPreRollExecutionToken
 
 		FMovieSceneMediaData& SectionData = PersistentData.GetSectionData<FMovieSceneMediaData>();
 		UMediaPlayer* MediaPlayer = SectionData.GetMediaPlayer();
-		IMediaPlayerProxyInterface* PlayerProxyInterface = Cast<IMediaPlayerProxyInterface>(SectionData.GetPlayerProxy());
+		const IMediaPlayerProxyInterface* PlayerProxyInterface = Cast<IMediaPlayerProxyInterface>(SectionData.GetPlayerProxy());
 		UMediaSource* MediaSource = GetMediaSource(Player, Operand.SequenceID);
 
 		if (MediaPlayer == nullptr || MediaSource == nullptr)
@@ -142,7 +142,7 @@ struct FMediaSectionPreRollExecutionToken
 			return;
 		}
 
-		bool bMoveToNewTime = Context.GetStatus() != EMovieScenePlayerStatus::Playing || (Context.GetStatus() == EMovieScenePlayerStatus::Playing && Context.HasJumped());
+		const bool bMoveToNewTime = Context.GetStatus() != EMovieScenePlayerStatus::Playing || (Context.GetStatus() == EMovieScenePlayerStatus::Playing && Context.HasJumped());
 		if (bMoveToNewTime)
 		{
 			MediaPlayer->SetRate(0.0f);
@@ -169,6 +169,18 @@ struct FMediaSectionExecutionToken
 		, bCanPlayerBeOpen(bInCanPlayerBeOpen)
 	{ }
 
+	static void SeekPlayer(const FMovieSceneContext& InContext, UMediaPlayer* InMediaPlayer, const FTimespan& InMediaTime)
+	{
+		if (InContext.GetStatus() == EMovieScenePlayerStatus::Scrubbing)
+		{
+			InMediaPlayer->Scrub(InMediaTime);
+		}
+		else
+		{
+			InMediaPlayer->Seek(InMediaTime);
+		}
+	}
+
 	virtual void Execute(const FMovieSceneContext& Context, const FMovieSceneEvaluationOperand& Operand, FPersistentEvaluationData& PersistentData, IMovieScenePlayer& Player) override
 	{
 		FMovieSceneMediaData& SectionData = PersistentData.GetSectionData<FMovieSceneMediaData>();
@@ -215,7 +227,7 @@ struct FMediaSectionExecutionToken
 		// Check if the cache options have changed mid playback.
 		//const FMediaSourceCacheSettings CurrentCacheSettings = GetBaseMediaSourceOptions().CacheSettings;
 		const FMediaSourceCacheSettings CurrentCacheSettings = FMediaSectionBaseMediaUtils::GetCurrentCacheSettingsFromPlayer(MediaPlayer);
-		bool bCacheSettingsChanged = PlayerProxyInterface ? CurrentCacheSettings != PlayerProxyInterface->GetCacheSettings() : false;
+		const bool bCacheSettingsChanged = PlayerProxyInterface ? CurrentCacheSettings != PlayerProxyInterface->GetCacheSettings() : false;
 
 		// open the media source if necessary
 		if (MediaPlayer->GetUrl().IsEmpty() || bCacheSettingsChanged)
@@ -270,7 +282,7 @@ struct FMediaSectionExecutionToken
 		{
 			if (!MediaPlayer->IsPlaying())
 			{
-				MediaPlayer->Seek(MediaTime);
+				SeekPlayer(Context, MediaPlayer, MediaTime);
 				// Set rate
 				// (note that the DIRECTION is important, but the magnitude is not - as we use blocked playback, the range setup to block on will serve as external clock to the player,
 				//  the direction is taken into account as hint for internal operation of the player)
@@ -286,10 +298,10 @@ struct FMediaSectionExecutionToken
 			{
 				if (Context.HasJumped())
 				{
-					MediaPlayer->Seek(MediaTime);
+					SeekPlayer(Context, MediaPlayer, MediaTime);
 				}
 
-				float CurrentPlayerRate = MediaPlayer->GetRate();
+				const float CurrentPlayerRate = MediaPlayer->GetRate();
 				if (Context.GetDirection() == EPlayDirection::Forwards && CurrentPlayerRate < 0.0f)
 				{
 					if (!MediaPlayer->SetRate(1.0f))
@@ -297,7 +309,7 @@ struct FMediaSectionExecutionToken
 						// Failed to set needed rate. Keep things blocked, as this means the player will still be returning the old rate, we will get here repeatedly
 						// and each time trigger a seek. A potentially very SLOW method of approximating backwards playback, but better
 						// than nothing.
-						MediaPlayer->Seek(MediaTime);
+						SeekPlayer(Context, MediaPlayer, MediaTime);
 					}
 				}
 				else if (Context.GetDirection() == EPlayDirection::Backwards && CurrentPlayerRate > 0.0f)
@@ -307,7 +319,7 @@ struct FMediaSectionExecutionToken
 						// Failed to set needed rate. Keep things blocked, as this means the player will still be returning the old rate, we will get here repeatedly
 						// and each time trigger a seek. A potentially very SLOW method of approximating backwards playback, but better
 						// than nothing.
-						MediaPlayer->Seek(MediaTime);
+						SeekPlayer(Context, MediaPlayer, MediaTime);
 					}
 				}
 			}
@@ -319,7 +331,7 @@ struct FMediaSectionExecutionToken
 				MediaPlayer->SetRate(0.0f);
 			}
 
-			MediaPlayer->Seek(MediaTime);
+			SeekPlayer(Context, MediaPlayer, MediaTime);
 		}
 
 		// Set blocking range / time-range to display
diff --git a/Engine/Source/Runtime/MediaAssets/Private/Assets/MediaPlayer.cpp b/Engine/Source/Runtime/MediaAssets/Private/Assets/MediaPlayer.cpp
--- a/Engine/Source/Runtime/MediaAssets/Private/Assets/MediaPlayer.cpp
+++ b/Engine/Source/Runtime/MediaAssets/Private/Assets/MediaPlayer.cpp
@@ -796,6 +796,16 @@ bool UMediaPlayer::Seek(const FTimespan& Time)
 	return PlayerFacade->Seek(Time);
 }
 
+bool UMediaPlayer::Scrub(const FTimespan& Time)
+{
+	if (PlayerFacade->IsSeeking())
+	{
+		UE_LOG(LogMediaAssets, VeryVerbose, TEXT("%s.Scrub %s"), *GetFName().ToString(), *Time.ToString(TEXT("%h:%m:%s.%t")));
+		PlayerFacade->SetNextSeek(Time);
+		return true;	// todo: add a failed seek event?
+	}
+	return Seek(Time);
+}
 
 bool UMediaPlayer::SelectTrack(EMediaPlayerTrack TrackType, int32 TrackIndex)
 {
diff --git a/Engine/Source/Runtime/MediaAssets/Public/MediaPlayer.h b/Engine/Source/Runtime/MediaAssets/Public/MediaPlayer.h
--- a/Engine/Source/Runtime/MediaAssets/Public/MediaPlayer.h
+++ b/Engine/Source/Runtime/MediaAssets/Public/MediaPlayer.h
@@ -372,7 +372,7 @@ public:
 	 * Get the supported playback rates.
 	 *
 	 * @param Unthinned Whether the rates are for unthinned playback.
-	 * @param Will contain the the ranges of supported rates.
+	 * @param OutRates Will contain the the ranges of supported rates.
 	 * @see SetRate, SupportsRate
 	 */
 	UFUNCTION(BlueprintCallable, Category="Media|MediaPlayer")
@@ -818,6 +818,25 @@ public:
 	UFUNCTION(BlueprintCallable, Category="Media|MediaPlayer")
 	MEDIAASSETS_API bool Seek(const FTimespan& Time);
 
+	/**
+	 * Seeks to the specified playback time as soon as possible.
+	 * 
+	 * If the player is currently seeking, the new seek request will be queued
+	 * and executed as soon as the current one is completed. It will guarantee that
+	 * at least one sample of the previous seek makes it to the sample queue.
+	 * if new seek requests are performed in rapid succession (as in scrubbing),
+	 * the requests are aggregated and only the most recent request will be performed.
+	 * As a result, OnSeekCompleted event is only sent for the request that got performed.
+	 *
+	 * If the player is not currently seeking, Seek will be called immediately.
+	 *
+	 * @param Time The playback time to set.
+	 * @return true on success, false otherwise.
+	 * @see Seek
+	 */
+	UFUNCTION(BlueprintCallable, Category="Media|MediaPlayer")
+	MEDIAASSETS_API bool Scrub(const FTimespan& Time);
+
 	/**
 	 * Select the active track of the given type.
 	 *
@@ -1281,7 +1300,7 @@ public:
 public:
 	/**
 	 * These are possible common information for use with GetMediaInfo.
-	 * They also decsribe the type of data that will be needed to retrieve the information.
+	 * They also describe the type of data that will be needed to retrieve the information.
 	 */
 	 /** Number of mip map levels in the source. Int32.*/
 	static MEDIAASSETS_API FLazyName MediaInfoNameSourceNumMips;
diff --git a/Engine/Source/Runtime/MediaUtils/Private/MediaPlayerFacade.cpp b/Engine/Source/Runtime/MediaUtils/Private/MediaPlayerFacade.cpp
--- a/Engine/Source/Runtime/MediaUtils/Private/MediaPlayerFacade.cpp
+++ b/Engine/Source/Runtime/MediaUtils/Private/MediaPlayerFacade.cpp
@@ -1139,6 +1139,7 @@ bool FMediaPlayerFacade::ContinueOpen(IMediaPlayerLifecycleManagerDelegate::ICon
 
 		NextEstVideoTimeAtFrameStart.Invalidate();
 		SeekTargetTime.Invalidate();
+		NextSeekTime.Reset();
 		SeekIndex = 0;
 	}
 
@@ -1183,6 +1184,8 @@ void FMediaPlayerFacade::QueryCacheState(EMediaTrackType TrackType, EMediaCacheS
 
 bool FMediaPlayerFacade::Seek(const FTimespan& InTime)
 {
+	NextSeekTime.Reset();
+
 	auto CurrentPlayer = Player;
 
 	if (!CurrentPlayer.IsValid())
@@ -1241,6 +1244,19 @@ bool FMediaPlayerFacade::Seek(const FTimespan& InTime)
 	return true;
 }
 
+bool FMediaPlayerFacade::IsSeeking() const
+{
+	// Code inspection notes:
+	// - Usually protected by LastTimeValuesCS, but sometimes CriticalSection (or both interlocked).
+	// - GetCurrentPlaybackTimeRange reads it outside of a scope lock.
+	FScopeLock Lock(&LastTimeValuesCS);
+	return SeekTargetTime.IsValid();
+}
+
+void FMediaPlayerFacade::SetNextSeek(const FTimespan& InTime)
+{
+	NextSeekTime = InTime;
+}
 
 void FMediaPlayerFacade::SetBlockOnTime(const FTimespan& Time)
 {
@@ -2414,6 +2430,8 @@ void FMediaPlayerFacade::TickOutput(FTimespan DeltaTime, FTimespan /*Timecode*/)
 	}
 
 	Cache->Tick(DeltaTime, CurrentRate, GetTime());
+
+	ExecuteNextSeek();
 }
 
 
@@ -2581,6 +2599,16 @@ void FMediaPlayerFacade::UpdateSeekStatus(const FMediaTimeStamp* pCheckTimeStamp
 	}
 }
 
+void FMediaPlayerFacade::ExecuteNextSeek()
+{
+	if (NextSeekTime.IsSet() && !IsSeeking())
+	{
+		if (!Seek(NextSeekTime.GetValue()))
+		{
+			// todo: signal error for failed seek.
+		}
+	}
+}
 
 void FMediaPlayerFacade::MonitorAudioEnablement()
 {
@@ -3652,6 +3680,7 @@ void FMediaPlayerFacade::ReInit()
 		NextEstVideoTimeAtFrameStart.Invalidate();
 		MostRecentlyDeliveredVideoFrameTimecode.Reset();
 		SeekTargetTime.Invalidate();
+		NextSeekTime.Reset();
 		SeekIndex = 0;
 	}
 }
diff --git a/Engine/Source/Runtime/MediaUtils/Public/MediaPlayerFacade.h b/Engine/Source/Runtime/MediaUtils/Public/MediaPlayerFacade.h
--- a/Engine/Source/Runtime/MediaUtils/Public/MediaPlayerFacade.h
+++ b/Engine/Source/Runtime/MediaUtils/Public/MediaPlayerFacade.h
@@ -615,6 +615,18 @@ public:
 	 */
 	MEDIAUTILS_API bool Seek(const FTimespan& Time);
 
+	/**
+	 * A seek has been requested, waiting to receive samples.
+	 * @see Seek
+	 */
+	MEDIAUTILS_API bool IsSeeking() const;
+
+	/**
+	 *	Set seek to be done after current seek is completed.
+	 *	@remark This value is reset if Seek is called directly.
+	 */
+	MEDIAUTILS_API void SetNextSeek(const FTimespan& InTime);
+	
 	/**
 	 * Select the active track of the given type.
 	 *
@@ -923,6 +935,7 @@ private:
 	MEDIAUTILS_API bool IsVideoSampleStillGood(const TRange<FMediaTimeStamp>& LastSampleTimeRange, const TRange<FMediaTimeStamp>& TimeRange, bool bReverse) const;
 	MEDIAUTILS_API void MonitorAudioEnablement();
 	MEDIAUTILS_API void UpdateSeekStatus(const FMediaTimeStamp* pCheckTimeStamp = nullptr);
+	void ExecuteNextSeek();
 	MEDIAUTILS_API void PreSampleProcessingTimeHandling();
 	MEDIAUTILS_API bool GetCurrentPlaybackTimeRange(TRange<FMediaTimeStamp>& TimeRange, float Rate, FTimespan DeltaTime, bool bPurgeSampleRelated) const;
 	MEDIAUTILS_API void PostSampleProcessingTimeHandling(FTimespan DeltaTime);
@@ -1083,6 +1096,9 @@ private:
 	/** Timestamp of seek target location if seek is pending */
 	FMediaTimeStamp SeekTargetTime;
 
+	/** When multiple consecutive seek are issued while seeking, we keep track of the most recent request. */
+	TOptional<FTimespan> NextSeekTime;
+	
 	/** Current seek index */
 	int32 SeekIndex;
 
