Movie Scene: Removed the reliance on actors and actor components for Media Player Proxies.

diff --git a/Engine/Plugins/Media/MediaCompositing/Source/MediaCompositing/Private/MovieScene/MovieSceneMediaSection.cpp b/Engine/Plugins/Media/MediaCompositing/Source/MediaCompositing/Private/MovieScene/MovieSceneMediaSection.cpp
--- a/Engine/Plugins/Media/MediaCompositing/Source/MediaCompositing/Private/MovieScene/MovieSceneMediaSection.cpp
+++ b/Engine/Plugins/Media/MediaCompositing/Source/MediaCompositing/Private/MovieScene/MovieSceneMediaSection.cpp
@@ -145,8 +145,7 @@ UMediaSource* UMovieSceneMediaSection::GetMediaSourceOrProxy(IMovieScenePlayer&
 	{
 		if (UObject* Object = WeakObject.Get())
 		{
-			AActor* Actor = Cast<AActor>(Object);
-			if (Actor != nullptr)
+			if (AActor* Actor = Cast<AActor>(Object))
 			{
 				// Loop over all the components and see if any are a proxy.
 				TArray<UActorComponent*> ActorComponents;
@@ -165,6 +164,11 @@ UMediaSource* UMovieSceneMediaSection::GetMediaSourceOrProxy(IMovieScenePlayer&
 					break;
 				}
 			}
+			else if (IMediaPlayerProxyInterface* Proxy = Cast<IMediaPlayerProxyInterface>(Object))
+			{
+				OutMediaSource = Proxy->ProxyGetMediaSourceFromIndex(InMediaSourceProxyIndex);
+				break;
+			}
 		}
 	}
 
