[MediaCompositing] Fix some Media Track issues when switching media source in the corresponding MediaPlate.

How this works: the media track creates an alternate media player in FMovieSceneMediaData for each media section as it is evaluated. It hooks up that player in the media texture (shared from the media plate), overriding the media plate's player.

3 fixes:
- in OnRegister, which is called when properties are edited, the media plate shouldn't override back the player if it is in proxy mode (i.e currently used to evaluate a track secion).
- Since the media plate's media source can be changed, the track section evaluation should check if the source has changed and if so, open it's internal player with the new source.
- The track editor will now listen for the media plate events to know if it should evaluate the track in case something changed.

diff --git a/Engine/Plugins/Media/MediaCompositing/Source/MediaCompositing/Private/MovieScene/MovieSceneMediaTemplate.cpp b/Engine/Plugins/Media/MediaCompositing/Source/MediaCompositing/Private/MovieScene/MovieSceneMediaTemplate.cpp
--- a/Engine/Plugins/Media/MediaCompositing/Source/MediaCompositing/Private/MovieScene/MovieSceneMediaTemplate.cpp
+++ b/Engine/Plugins/Media/MediaCompositing/Source/MediaCompositing/Private/MovieScene/MovieSceneMediaTemplate.cpp
@@ -123,7 +123,7 @@ struct FMediaSectionPreRollExecutionToken
 		}
 
 		// open the media source if necessary
-		if (MediaPlayer->GetUrl().IsEmpty())
+		if (MediaPlayer->GetUrl() != MediaSource->GetUrl())
 		{
 			SectionData.SeekOnOpen(StartTime);
 
@@ -230,7 +230,7 @@ struct FMediaSectionExecutionToken
 		const bool bCacheSettingsChanged = PlayerProxyInterface ? CurrentCacheSettings != PlayerProxyInterface->GetCacheSettings() : false;
 
 		// open the media source if necessary
-		if (MediaPlayer->GetUrl().IsEmpty() || bCacheSettingsChanged)
+		if (MediaPlayer->GetUrl() != MediaSource->GetUrl() || bCacheSettingsChanged)
 		{
 			SectionData.SeekOnOpen(CurrentTime);
 
diff --git a/Engine/Plugins/Media/MediaCompositing/Source/MediaCompositingEditor/Private/Sequencer/MediaTrackEditor.cpp b/Engine/Plugins/Media/MediaCompositing/Source/MediaCompositingEditor/Private/Sequencer/MediaTrackEditor.cpp
--- a/Engine/Plugins/Media/MediaCompositing/Source/MediaCompositingEditor/Private/Sequencer/MediaTrackEditor.cpp
+++ b/Engine/Plugins/Media/MediaCompositing/Source/MediaCompositingEditor/Private/Sequencer/MediaTrackEditor.cpp
@@ -48,6 +48,9 @@ FMediaTrackEditor::FMediaTrackEditor(TSharedRef<ISequencer> InSequencer)
 {
 	ThumbnailPool = MakeShared<FTrackEditorThumbnailPool>(InSequencer);
 	InSequencer->OnMovieSceneDataChanged().AddRaw(this, &FMediaTrackEditor::OnSequencerDataChanged);
+
+	IMediaAssetsModule& MediaAssetsModule = FModuleManager::Get().LoadModuleChecked<IMediaAssetsModule>("MediaAssets");
+	OnMediaPlateStateChangedHandle = MediaAssetsModule.RegisterOnMediaStateChangedEvent(IMediaAssetsModule::FMediaStateChangedDelegate::FDelegate::CreateRaw(this, &FMediaTrackEditor::OnMediaPlateStateChanged));
 }
 
 
@@ -58,6 +61,9 @@ FMediaTrackEditor::~FMediaTrackEditor()
 	{
 		SequencerPtr->OnMovieSceneDataChanged().RemoveAll(this);
 	}
+	
+	IMediaAssetsModule& MediaAssetsModule = FModuleManager::Get().LoadModuleChecked<IMediaAssetsModule>("MediaAssets");
+	MediaAssetsModule.UnregisterOnMediaStateChangedEvent(OnMediaPlateStateChangedHandle);
 }
 
 
@@ -383,7 +389,7 @@ void FMediaTrackEditor::OnSequencerDataChanged(EMovieSceneDataChangeType DataCha
 void FMediaTrackEditor::UpdateMediaTrackBinding(UMovieSceneMediaTrack* MediaTrack, const FMovieSceneBinding& Binding)
 {
 	// Get object from the binding.
-	FGuid ObjectHandle = Binding.GetObjectGuid();
+	const FGuid ObjectHandle = Binding.GetObjectGuid();
 	const TSharedPtr<ISequencer> SequencerPtr = GetSequencer();
 	UObject* BoundObject = SequencerPtr.IsValid() ?
 		SequencerPtr->FindSpawnedObjectOrTemplate(ObjectHandle) : nullptr;
@@ -398,7 +404,7 @@ void FMediaTrackEditor::UpdateMediaTrackBinding(UMovieSceneMediaTrack* MediaTrac
 			MediaAssetsModule->GetPlayerFromObject(BoundObject, PlayerProxy);
 		}
 	}
-	bool bHasPlayerProxy = PlayerProxy != nullptr;
+	const bool bHasPlayerProxy = PlayerProxy != nullptr;
 
 	// Update all sections.
 	const TArray<UMovieSceneSection*> Sections = MediaTrack->GetAllSections();
@@ -412,4 +418,34 @@ void FMediaTrackEditor::UpdateMediaTrackBinding(UMovieSceneMediaTrack* MediaTrac
 	}
 }
 
+void FMediaTrackEditor::OnMediaPlateStateChanged(const TArray<FString>& InActorsPathNames, uint8 InEnumState, bool bRemoteBroadcast)
+{
+	// Ignore remote events.
+	if (bRemoteBroadcast)
+	{
+		return;
+	}
+
+	const TSharedPtr<ISequencer> SequencerPtr = GetSequencer();
+	const UMovieSceneSequence* MovieSequence = SequencerPtr ? SequencerPtr->GetFocusedMovieSceneSequence() : nullptr;
+	const UMovieScene* MovieScene = MovieSequence ? MovieSequence->GetMovieScene() : nullptr;
+	if (!MovieScene)
+	{
+		return;
+	}
+
+	for (const FMovieSceneBinding& Binding : MovieScene->GetBindings())
+	{
+		FGuid ObjectHandle = Binding.GetObjectGuid();
+		if (const UObject* BoundObject = SequencerPtr->FindSpawnedObjectOrTemplate(ObjectHandle))
+		{
+			if (InActorsPathNames.Contains(BoundObject->GetPathName()))
+			{
+				SequencerPtr->RequestEvaluate();
+				break;
+			}
+		}
+	}
+}
+
 #undef LOCTEXT_NAMESPACE
diff --git a/Engine/Plugins/Media/MediaCompositing/Source/MediaCompositingEditor/Public/Sequencer/MediaTrackEditor.h b/Engine/Plugins/Media/MediaCompositing/Source/MediaCompositingEditor/Public/Sequencer/MediaTrackEditor.h
--- a/Engine/Plugins/Media/MediaCompositing/Source/MediaCompositingEditor/Public/Sequencer/MediaTrackEditor.h
+++ b/Engine/Plugins/Media/MediaCompositing/Source/MediaCompositingEditor/Public/Sequencer/MediaTrackEditor.h
@@ -102,7 +102,10 @@ private:
 	 */
 	void UpdateMediaTrackBinding(UMovieSceneMediaTrack* MediaTrack, const FMovieSceneBinding& Binding);
 
+	void OnMediaPlateStateChanged(const TArray<FString>& InActorsPathNames, uint8 InEnumState, bool bInRemoteBroadcast);
+
 private:
 
 	TSharedPtr<FTrackEditorThumbnailPool> ThumbnailPool;
+	FDelegateHandle OnMediaPlateStateChangedHandle;
 };
diff --git a/Engine/Plugins/Media/MediaPlate/Source/MediaPlate/Private/MediaPlateComponent.cpp b/Engine/Plugins/Media/MediaPlate/Source/MediaPlate/Private/MediaPlateComponent.cpp
--- a/Engine/Plugins/Media/MediaPlate/Source/MediaPlate/Private/MediaPlateComponent.cpp
+++ b/Engine/Plugins/Media/MediaPlate/Source/MediaPlate/Private/MediaPlateComponent.cpp
@@ -173,8 +173,8 @@ void UMediaPlateComponent::OnRegister()
 	MediaPlayer->OnPlaybackResumed.AddUniqueDynamic(this, &UMediaPlateComponent::OnMediaResumed);
 	MediaPlayer->OnPlaybackSuspended.AddUniqueDynamic(this, &UMediaPlateComponent::OnMediaSuspended);
 
-	// Set up media texture.
-	SetUpTextures(UE::MediaPlateComponent::ESetUpTexturesFlags::AllowSetPlayer);
+	// Set up media texture. Note: don't set player if in proxy mode.
+	SetUpTextures(IsNormalMode() ? UE::MediaPlateComponent::ESetUpTexturesFlags::AllowSetPlayer : UE::MediaPlateComponent::ESetUpTexturesFlags::None);
 
 	// Set up sound component if we have one.
 	if (SoundComponent != nullptr)
@@ -1232,6 +1232,15 @@ void UMediaPlateComponent::SetNormalMode(bool bInIsNormalMode)
 #endif // WITH_EDITOR
 }
 
+bool UMediaPlateComponent::IsNormalMode() const
+{
+#if WITH_EDITOR
+	return bIsNormalMode;
+#else
+	return true;
+#endif
+}
+
 void UMediaPlateComponent::UpdateTextureLayers()
 {
 	if (AMediaPlate* MediaPlate = GetOwner<AMediaPlate>())
diff --git a/Engine/Plugins/Media/MediaPlate/Source/MediaPlate/Public/MediaPlateComponent.h b/Engine/Plugins/Media/MediaPlate/Source/MediaPlate/Public/MediaPlateComponent.h
--- a/Engine/Plugins/Media/MediaPlate/Source/MediaPlate/Public/MediaPlateComponent.h
+++ b/Engine/Plugins/Media/MediaPlate/Source/MediaPlate/Public/MediaPlateComponent.h
@@ -570,6 +570,11 @@ private:
 	 */
 	void SetNormalMode(bool bInIsNormalMode);
 
+	/**
+	 * Returns true if normal mode of false if proxy mode.
+	 */
+	bool IsNormalMode() const;
+	
 	/**
 	 * Sets textures in our material according to the layer assignments.
 	 */
