[Media] Media Track: Fix Preroll's hitch by ensuring there is no extra seek at the start of the playback section (coming out of preroll section).

diff --git a/Engine/Plugins/Media/MediaCompositing/Source/MediaCompositing/MediaCompositing.Build.cs b/Engine/Plugins/Media/MediaCompositing/Source/MediaCompositing/MediaCompositing.Build.cs
--- a/Engine/Plugins/Media/MediaCompositing/Source/MediaCompositing/MediaCompositing.Build.cs
+++ b/Engine/Plugins/Media/MediaCompositing/Source/MediaCompositing/MediaCompositing.Build.cs
@@ -25,6 +25,7 @@ namespace UnrealBuildTool.Rules
 				new string[] {
 					"Media",
 					"MediaAssets",
+					"MediaUtils",
 					"RenderCore",
 					"RHI",
 				});
diff --git a/Engine/Plugins/Media/MediaCompositing/Source/MediaCompositing/Private/MovieScene/MovieSceneMediaTemplate.cpp b/Engine/Plugins/Media/MediaCompositing/Source/MediaCompositing/Private/MovieScene/MovieSceneMediaTemplate.cpp
--- a/Engine/Plugins/Media/MediaCompositing/Source/MediaCompositing/Private/MovieScene/MovieSceneMediaTemplate.cpp
+++ b/Engine/Plugins/Media/MediaCompositing/Source/MediaCompositing/Private/MovieScene/MovieSceneMediaTemplate.cpp
@@ -3,6 +3,7 @@
 #include "MovieSceneMediaTemplate.h"
 
 #include "IMediaAssetsModule.h"
+#include "IMediaCache.h"
 #include "Math/UnrealMathUtility.h"
 #include "MediaPlayer.h"
 #include "MediaPlayerFacade.h"
@@ -55,6 +56,51 @@ namespace FMediaSectionBaseMediaUtils
 		}
 		return cs;
 	}
+
+	/**
+	 * Check if we have a sample available for the desired time, either in the sample queue or already enqueued in video sink.
+	 */
+	bool IsVideoSampleAvailableForTime(const TSharedRef<FMediaPlayerFacade>& InMediaPlayerFacade, const FTimespan& InTime)
+	{
+		// Todo: For Protron, this doesn't work on loops. Only works for preroll.
+		// Preroll is caught by the first test. But doesn't work on loops.
+		// For loops, the QueryCacheState should catch it, but Protron never has any "loaded" samples.
+		
+		// Check LastVideoSampleProcessedTimeRange to see if the sample has already been consumed.
+		const TRange<FMediaTimeStamp> LastVideoSampleTimeRange = InMediaPlayerFacade->GetLastProcessedVideoSampleTimeRange();
+		if (!LastVideoSampleTimeRange.IsEmpty())
+		{
+			// Remark: ignoring sequence and loop indices for now.
+			const TRange<FTimespan> TimeRange(LastVideoSampleTimeRange.GetLowerBoundValue().Time, LastVideoSampleTimeRange.GetUpperBoundValue().Time);
+			if (TimeRange.Contains(InTime))
+			{
+				return true;
+			}
+		}
+
+		// Query the cache state.
+		TRangeSet<FTimespan> CacheRangeSet;
+		InMediaPlayerFacade->QueryCacheState(EMediaTrackType::Video, EMediaCacheState::Loaded, CacheRangeSet);
+		return CacheRangeSet.Contains(InTime);
+	}
+
+	/**
+	 * Checks if the player needs to seek to specified time.
+	 */
+	bool ShouldSeekForTime(const UMediaPlayer* InMediaPlayer, const FTimespan& InTime)
+	{
+		const TSharedRef<FMediaPlayerFacade> MediaPlayerFacade = InMediaPlayer->GetPlayerFacade();
+
+		// Check if already seeking to desired target.
+		const FMediaTimeStamp SeekTarget = MediaPlayerFacade->GetSeekTarget();
+		if (SeekTarget.IsValid() && SeekTarget.Time == InTime)
+		{
+			return false;	// Don't need to seek if already seeking to that specified time.
+		}
+
+		// Otherwise, check if we have a sample available for the desired time, either in the sample queue or already processed.
+		return IsVideoSampleAvailableForTime(MediaPlayerFacade, InTime);
+	}
 }
 
 /** Base struct for execution tokens. */
@@ -282,7 +328,12 @@ struct FMediaSectionExecutionToken
 		{
 			if (!MediaPlayer->IsPlaying())
 			{
-				SeekPlayer(Context, MediaPlayer, MediaTime);
+				// If the player has pre-rolled, it is ready and doesn't need to seek.
+				if (!FMediaSectionBaseMediaUtils::ShouldSeekForTime(MediaPlayer, MediaTime))
+				{
+					SeekPlayer(Context, MediaPlayer, MediaTime);
+				}
+
 				// Set rate
 				// (note that the DIRECTION is important, but the magnitude is not - as we use blocked playback, the range setup to block on will serve as external clock to the player,
 				//  the direction is taken into account as hint for internal operation of the player)
@@ -296,7 +347,8 @@ struct FMediaSectionExecutionToken
 			}
 			else
 			{
-				if (Context.HasJumped())
+				// Todo: avoid seek on loop. Causes loop index to be different between player and facade.
+				if (Context.HasJumped()) // && !FMediaSectionBaseMediaUtils::ShouldSeekForTime(MediaPlayer, MediaTime))
 				{
 					SeekPlayer(Context, MediaPlayer, MediaTime);
 				}
diff --git a/Engine/Source/Runtime/MediaUtils/Private/MediaPlayerFacade.cpp b/Engine/Source/Runtime/MediaUtils/Private/MediaPlayerFacade.cpp
--- a/Engine/Source/Runtime/MediaUtils/Private/MediaPlayerFacade.cpp
+++ b/Engine/Source/Runtime/MediaUtils/Private/MediaPlayerFacade.cpp
@@ -520,6 +520,11 @@ TOptional<FTimecode> FMediaPlayerFacade::GetVideoTimecode() const
 	return MostRecentlyDeliveredVideoFrameTimecode;
 }
 
+TRange<FMediaTimeStamp> FMediaPlayerFacade::GetLastProcessedVideoSampleTimeRange() const
+{
+	FScopeLock Lock(&LastTimeValuesCS);
+	return LastVideoSampleProcessedTimeRange;
+}
 
 FMediaTimeStamp FMediaPlayerFacade::GetTimeStampInternal(bool bForDisplay) const
 {
@@ -1291,6 +1296,12 @@ bool FMediaPlayerFacade::IsSeeking() const
 	return SeekTargetTime.IsValid();
 }
 
+FMediaTimeStamp FMediaPlayerFacade::GetSeekTarget() const
+{
+	FScopeLock Lock(&LastTimeValuesCS);
+	return SeekTargetTime;
+}
+
 void FMediaPlayerFacade::SetNextSeek(const FTimespan& InTime)
 {
 	NextSeekTime = InTime;
diff --git a/Engine/Source/Runtime/MediaUtils/Public/MediaPlayerFacade.h b/Engine/Source/Runtime/MediaUtils/Public/MediaPlayerFacade.h
--- a/Engine/Source/Runtime/MediaUtils/Public/MediaPlayerFacade.h
+++ b/Engine/Source/Runtime/MediaUtils/Public/MediaPlayerFacade.h
@@ -406,6 +406,11 @@ public:
 	 */
 	MEDIAUTILS_API TOptional<FTimecode> GetVideoTimecode() const;
 
+	/**
+	 *	Returns the last video sample time range that has been enqueued in the video sink.
+	 */
+	MEDIAUTILS_API TRange<FMediaTimeStamp> GetLastProcessedVideoSampleTimeRange() const;
+
 	/**
 	 * Get the human readable name of the specified track.
 	 *
@@ -620,7 +625,13 @@ public:
 	 * @see Seek
 	 */
 	MEDIAUTILS_API bool IsSeeking() const;
-
+	
+	/**
+	 * If valid, returns the current seek request.
+	 * @see Seek, IsSeeking
+	 */
+	MEDIAUTILS_API FMediaTimeStamp GetSeekTarget() const;
+	
 	/**
 	 *	Set seek to be done after current seek is completed.
 	 *	@remark This value is reset if Seek is called directly.
