[MediaCompositing] Media Thumbnail section: cache line not appearing when sections are too small.

diff --git a/Engine/Plugins/Media/MediaCompositing/Source/MediaCompositingEditor/Private/Sequencer/MediaThumbnailSection.cpp b/Engine/Plugins/Media/MediaCompositing/Source/MediaCompositingEditor/Private/Sequencer/MediaThumbnailSection.cpp
--- a/Engine/Plugins/Media/MediaCompositing/Source/MediaCompositingEditor/Private/Sequencer/MediaThumbnailSection.cpp
+++ b/Engine/Plugins/Media/MediaCompositing/Source/MediaCompositingEditor/Private/Sequencer/MediaThumbnailSection.cpp
@@ -358,15 +358,14 @@ void FMediaThumbnailSection::DrawSampleStates(FSequencerSectionPainter& InPainte
 
 	TArray<TRange<FTimespan>> Ranges;
 	RangeSet.GetRanges(Ranges);
-	float LoopDrawOffset = -TimeToPixelConverter.SecondsToPixel(TickResolution.AsSeconds(MediaSection->StartFrameOffset));
+	float LoopDrawOffset = -TimeToPixelConverter.SecondsDeltaToPixel(TickResolution.AsSeconds(MediaSection->StartFrameOffset));
 
 	while (LoopDrawOffset < SectionSize.X)
 	{
 		for (auto& Range : Ranges)
 		{
 			const float DrawOffset = static_cast<float>(FMath::RoundToNegativeInfinity(FTimespan::Ratio(Range.GetLowerBoundValue(), MediaDuration) * MediaSizeX) +	
-				LoopDrawOffset +
-				TimeToPixelConverter.SecondsToPixel(0.0));
+				LoopDrawOffset);
 			const float DrawSize = static_cast<float>(FMath::RoundToPositiveInfinity(FTimespan::Ratio(Range.Size<FTimespan>(), MediaDuration) * MediaSizeX));
 			const float BarHeight = 4.0f;
 
