[MediaCompositing] Performance issues when dragging a Media Track

When a track is modified, under this api, the entire section eval template is regenerated. This leads to FMovieSceneMediaData (persistent eval data associated to that template) being destroyed, closing the media player along with it.

The inefficiency comes from closing the media player and having to open a new one for the new evaluation template.

Workaround:

When the FMovieSceneMediaData are being torn down, instead of closing the media player, it will store it in a temporary map (FMovieSceneMediaPlayerStore) associated to the owning media section object. Then when the new FMovieSceneMediaData is created, it can retrieve the media player. That way, the media player doesn't have to be closed and is instead passed to the newly generated FMovieSceneMediaData.

In FMovieSceneTrackEvaluator::CallSetupTearDown, the expired entities are torn down first, and the NewEntities are setup next. This is why, at the end of the evaluation, we can close all remaining media players in FMovieSceneMediaPlayerStore, assuming there won't be reused because the corresponding section is not being evaluated anymore.

The FMovieSceneMediaPlayerStore is itself kept around as a persistent evaluation data with a FSharedPersistentDataKey.

diff --git a/Engine/Plugins/Media/MediaCompositing/Source/MediaCompositing/Private/MovieScene/MovieSceneMediaData.cpp b/Engine/Plugins/Media/MediaCompositing/Source/MediaCompositing/Private/MovieScene/MovieSceneMediaData.cpp
--- a/Engine/Plugins/Media/MediaCompositing/Source/MediaCompositing/Private/MovieScene/MovieSceneMediaData.cpp
+++ b/Engine/Plugins/Media/MediaCompositing/Source/MediaCompositing/Private/MovieScene/MovieSceneMediaData.cpp
@@ -6,6 +6,7 @@
 #include "MediaPlayer.h"
 #include "MediaPlayerProxyInterface.h"
 #include "MediaTexture.h"
+#include "MovieSceneMediaPlayerStore.h"
 #include "UObject/Package.h"
 #include "UObject/UObjectGlobals.h"
 
@@ -17,6 +18,7 @@ FMovieSceneMediaData::FMovieSceneMediaData()
 	: bIsAspectRatioSet(false)
 	, bOverrideMediaPlayer(false)
 	, MediaPlayer(nullptr)
+	, ProxyLayerIndex(0)
 	, ProxyTextureIndex(0)
 	, SeekOnOpenTime(FTimespan::MinValue())
 { }
@@ -27,12 +29,16 @@ FMovieSceneMediaData::~FMovieSceneMediaData()
 	if (MediaPlayer != nullptr)
 	{
 		MediaPlayer->OnMediaEvent().RemoveAll(this);
-		MediaPlayer->Close();
-		if (!bOverrideMediaPlayer)
+
+		if (const TSharedPtr<FMovieSceneMediaPlayerStore> MediaPlayerStore = MediaPlayerStoreWeak.Pin())
+		{
+			MediaPlayerStore->ScheduleMediaPlayerForRelease(MediaSection, MediaPlayer);
+		}
+		else
 		{
-			MediaPlayer->CleanUpBeforeDestroy();
+			// If we don't have a store (ex: external player), close media player immediately.
+			FMovieSceneMediaPlayerStore::CloseMediaPlayer(MediaPlayer, !bOverrideMediaPlayer);
 		}
-		MediaPlayer->RemoveFromRoot();
 	}
 }
 
@@ -52,7 +58,7 @@ void FMovieSceneMediaData::SeekOnOpen(FTimespan Time)
 }
 
 
-void FMovieSceneMediaData::Setup(UMediaPlayer* OverrideMediaPlayer, UObject* InPlayerProxy, int32 InProxyLayerIndex, int32 InProxyTextureIndex)
+void FMovieSceneMediaData::Setup(const TSharedPtr<FMovieSceneMediaPlayerStore>& InMediaPlayerStore, const UObject* InMediaSection, UMediaPlayer* OverrideMediaPlayer, UObject* InPlayerProxy, int32 InProxyLayerIndex, int32 InProxyTextureIndex)
 {
 	// Ensure we don't already have a media player set. Setup should only be called once
 	check(!MediaPlayer);
@@ -64,7 +70,14 @@ void FMovieSceneMediaData::Setup(UMediaPlayer* OverrideMediaPlayer, UObject* InP
 	}
 	else if (MediaPlayer == nullptr)
 	{
-		MediaPlayer = NewObject<UMediaPlayer>(GetTransientPackage(), MakeUniqueObjectName(GetTransientPackage(), UMediaPlayer::StaticClass()));
+		MediaPlayerStoreWeak = InMediaPlayerStore;
+		MediaSection = InMediaSection;
+		// Try to acquire a player from the store associated to the owning section.
+		MediaPlayer = InMediaPlayerStore ? InMediaPlayerStore->TryAcquireMediaPlayer(InMediaSection) : nullptr;
+		if (!MediaPlayer)
+		{
+			MediaPlayer = NewObject<UMediaPlayer>(GetTransientPackage(), MakeUniqueObjectName(GetTransientPackage(), UMediaPlayer::StaticClass()));
+		}
 	}
 
 	MediaPlayer->PlayOnOpen = false;
diff --git a/Engine/Plugins/Media/MediaCompositing/Source/MediaCompositing/Private/MovieScene/MovieSceneMediaPlayerStore.cpp b/Engine/Plugins/Media/MediaCompositing/Source/MediaCompositing/Private/MovieScene/MovieSceneMediaPlayerStore.cpp
new file mode 100644
--- /dev/null
+++ b/Engine/Plugins/Media/MediaCompositing/Source/MediaCompositing/Private/MovieScene/MovieSceneMediaPlayerStore.cpp
@@ -0,0 +1,85 @@
+// Copyright Epic Games, Inc. All Rights Reserved.
+
+
+#include "MovieSceneMediaPlayerStore.h"
+
+#include "MediaPlayer.h"
+#include "Misc/CoreDelegates.h"
+
+namespace UE::MediaCompositing::Private
+{
+	static FSharedPersistentDataKey MediaPlayerDataContainerSharedKey(FMovieSceneSharedDataId::Allocate(), FMovieSceneEvaluationOperand());
+}
+
+FMovieSceneMediaPlayerStore::FMovieSceneMediaPlayerStore()
+{
+	OnEndFrameDelegateHandle = FCoreDelegates::OnEndFrame.AddRaw(this, &FMovieSceneMediaPlayerStore::OnEndFrame);
+}
+
+FMovieSceneMediaPlayerStore::~FMovieSceneMediaPlayerStore()
+{
+	FCoreDelegates::OnEndFrame.Remove(OnEndFrameDelegateHandle);
+	CloseRemainingPlayers();
+}
+
+void FMovieSceneMediaPlayerStore::ScheduleMediaPlayerForRelease(const UObject* InPersistentObject, UMediaPlayer* InMediaPlayer)
+{
+	// Validation code -- There shouldn't be more than one player per section in the current design.
+	if (UMediaPlayer** ExistingMediaPlayer = MediaPlayers.Find(InPersistentObject))
+	{
+		if (IsValid(*ExistingMediaPlayer) && *ExistingMediaPlayer != InMediaPlayer)
+		{
+			CloseMediaPlayer(*ExistingMediaPlayer);
+		}
+	}
+	
+	MediaPlayers.Add(InPersistentObject, InMediaPlayer);
+}
+
+UMediaPlayer* FMovieSceneMediaPlayerStore::TryAcquireMediaPlayer(const UObject* InPersistentObject)
+{
+	if (UMediaPlayer** FoundMediaPlayer = MediaPlayers.Find(InPersistentObject))
+	{
+		UMediaPlayer* MediaPlayer = *FoundMediaPlayer;
+		MediaPlayers.Remove(InPersistentObject);
+		return MediaPlayer;
+	}
+	return nullptr;
+}
+
+void FMovieSceneMediaPlayerStore::CloseMediaPlayer(UMediaPlayer* InMediaPlayer, bool bInCleanUpBeforeDestroy)
+{
+	if (IsValid(InMediaPlayer))
+	{
+		InMediaPlayer->Close();
+		if (bInCleanUpBeforeDestroy)
+		{
+			InMediaPlayer->CleanUpBeforeDestroy();
+		}
+		InMediaPlayer->RemoveFromRoot();
+	}
+}
+
+void FMovieSceneMediaPlayerStore::OnEndFrame()
+{
+	CloseRemainingPlayers();
+}
+
+void FMovieSceneMediaPlayerStore::CloseRemainingPlayers()
+{
+	for (const TPair<FObjectKey, UMediaPlayer*>& MediaPlayerEntry : MediaPlayers)
+	{
+		CloseMediaPlayer(MediaPlayerEntry.Value);
+	}
+	MediaPlayers.Reset();
+}
+
+FMovieSceneMediaPlayerStoreContainer::FMovieSceneMediaPlayerStoreContainer()
+{
+	MediaPlayerStore = MakeShared<FMovieSceneMediaPlayerStore>();
+}
+
+FMovieSceneMediaPlayerStoreContainer& FMovieSceneMediaPlayerStoreContainer::GetOrAdd(FPersistentEvaluationData& InPersistentData)
+{
+	return InPersistentData.GetOrAdd<FMovieSceneMediaPlayerStoreContainer>(UE::MediaCompositing::Private::MediaPlayerDataContainerSharedKey);
+}
diff --git a/Engine/Plugins/Media/MediaCompositing/Source/MediaCompositing/Private/MovieScene/MovieSceneMediaPlayerStore.h b/Engine/Plugins/Media/MediaCompositing/Source/MediaCompositing/Private/MovieScene/MovieSceneMediaPlayerStore.h
new file mode 100644
--- /dev/null
+++ b/Engine/Plugins/Media/MediaCompositing/Source/MediaCompositing/Private/MovieScene/MovieSceneMediaPlayerStore.h
@@ -0,0 +1,86 @@
+// Copyright Epic Games, Inc. All Rights Reserved.
+
+#pragma once
+
+#include "Evaluation/PersistentEvaluationData.h"
+
+class UMediaPlayer;
+
+/**
+ * Keep a temporary store of media players so they can be reused after recompiling a media section.
+ * During the track evaluation, the expired templates will be torn down first, then the new templates
+ * get initialized.
+ *
+ * The torn down templates will submit their media player in this store associated with the owning
+ * media section, so that the new template data (for the media section) can reuse the same player.
+ * If players are not reused by the end of the evaluation, it means the media section is not evaluating anymore and
+ * the remaining players can be closed and discarded.
+ */
+class FMovieSceneMediaPlayerStore
+{
+public:
+	FMovieSceneMediaPlayerStore();
+	virtual ~FMovieSceneMediaPlayerStore();
+
+	/**
+	 * Associate the given media player to the persistent object.
+	 * If not acquired again with the given persistent object by the end of the frame,
+	 * it will be closed and discarded.
+	 * 
+	 * @param InPersistentObject Persistent object (ex section) we want to associate the player with.
+	 * @param InMediaPlayer Media player to reuse.
+	 */
+	void ScheduleMediaPlayerForRelease(const UObject* InPersistentObject, UMediaPlayer* InMediaPlayer);
+
+	/**
+	 * @brief Get the media player registered to the given persistent object. Remove from pool if found.
+	 * @param InPersistentObject Persistent object (ex section) the media player is associated with.
+	 * @return pointer to media player if found, null otherwise.
+	 */
+	UMediaPlayer* TryAcquireMediaPlayer(const UObject* InPersistentObject);
+
+	/**
+	 * @brief Close the given media player and remove from root.
+	 * @param InMediaPlayer Input media player, validity is checked.
+	 * @param bInCleanUpBeforeDestroy Optionally calls CleanUpBeforeDestroy() to ensure player resources a freed immediately.
+	 */
+	static void CloseMediaPlayer(UMediaPlayer* InMediaPlayer, bool bInCleanUpBeforeDestroy = true);
+
+private:
+	/** Delegate handler for OneEndFrame. This is used to hook up CloseRemainingPlayers() somewhere after the end of the sequence evaluation. */
+	void OnEndFrame();
+
+	/** Called after the evaluation is completed to close all non-acquired players. */
+	void CloseRemainingPlayers();
+
+	/** Map of media players. */
+	TMap<FObjectKey, UMediaPlayer*> MediaPlayers;
+
+	FDelegateHandle OnEndFrameDelegateHandle;
+};
+
+/**
+ * Persistent Evaluation Data Wrapper for the media player association store.
+ * Using this pattern to be able to get a shared/weak ptr to the media player store so the media section data,
+ * which is also a persistent evaluation data, can keep a weak ptr to solving the order of destruction issue.
+ */
+class FMovieSceneMediaPlayerStoreContainer : public IPersistentEvaluationData
+{
+public:
+	FMovieSceneMediaPlayerStoreContainer();
+	virtual ~FMovieSceneMediaPlayerStoreContainer() override = default;
+
+	/**
+	 * Get the media player data container from the persistent data. Will be created if missing.
+	 */
+	static FMovieSceneMediaPlayerStoreContainer& GetOrAdd(FPersistentEvaluationData& InPersistentData);
+
+	/** Get the shared media player store. */
+	TSharedPtr<FMovieSceneMediaPlayerStore> GetMediaPlayerStore() const
+	{
+		return MediaPlayerStore;
+	}
+
+private:
+	TSharedPtr<FMovieSceneMediaPlayerStore> MediaPlayerStore;
+};
diff --git a/Engine/Plugins/Media/MediaCompositing/Source/MediaCompositing/Private/MovieScene/MovieSceneMediaTemplate.cpp b/Engine/Plugins/Media/MediaCompositing/Source/MediaCompositing/Private/MovieScene/MovieSceneMediaTemplate.cpp
--- a/Engine/Plugins/Media/MediaCompositing/Source/MediaCompositing/Private/MovieScene/MovieSceneMediaTemplate.cpp
+++ b/Engine/Plugins/Media/MediaCompositing/Source/MediaCompositing/Private/MovieScene/MovieSceneMediaTemplate.cpp
@@ -16,6 +16,7 @@
 #include "UObject/GCObject.h"
 
 #include "MovieSceneMediaData.h"
+#include "MovieSceneMediaPlayerStore.h"
 #include "MovieSceneMediaSection.h"
 #include "MovieSceneMediaTrack.h"
 
@@ -222,7 +223,7 @@ struct FMediaSectionExecutionToken
 
 	static void SeekPlayer(const FMovieSceneContext& InContext, UMediaPlayer* InMediaPlayer, const FTimespan& InMediaTime)
 	{
-		if (InContext.GetStatus() == EMovieScenePlayerStatus::Scrubbing)
+		if (InContext.GetStatus() == EMovieScenePlayerStatus::Scrubbing || InContext.GetStatus() == EMovieScenePlayerStatus::Stopped)
 		{
 			InMediaPlayer->Scrub(InMediaTime);
 		}
@@ -391,7 +392,8 @@ struct FMediaSectionExecutionToken
 			SeekPlayer(Context, MediaPlayer, MediaTime);
 		}
 
-		if (Context.GetStatus() == EMovieScenePlayerStatus::Scrubbing && CVarBlockOnFetchScrub.GetValueOnGameThread() == false)
+		if ((Context.GetStatus() == EMovieScenePlayerStatus::Scrubbing || Context.GetStatus() == EMovieScenePlayerStatus::Stopped)
+			&& CVarBlockOnFetchScrub.GetValueOnGameThread() == false)
 		{
 			// When scrubbing, seek requests are non-blocking.
 			MediaPlayer->SetBlockOnTimeRange(TRange<FTimespan>::Empty());
@@ -548,7 +550,8 @@ void FMovieSceneMediaSectionTemplate::Initialize(const FMovieSceneEvaluationOper
 
 		// Add section data.
 		SectionData = &PersistentData.AddSectionData<FMovieSceneMediaData>();
-		SectionData->Setup(MediaPlayer, PlayerProxy, ProxyLayerIndex, ProxyTextureIndex);
+		const TSharedPtr<FMovieSceneMediaPlayerStore> MediaPlayerStore = FMovieSceneMediaPlayerStoreContainer::GetOrAdd(PersistentData).GetMediaPlayerStore();
+		SectionData->Setup(MediaPlayerStore, MediaSection, MediaPlayer, PlayerProxy, ProxyLayerIndex, ProxyTextureIndex);
 	}
 
 	if (!ensure(SectionData != nullptr))
diff --git a/Engine/Plugins/Media/MediaCompositing/Source/MediaCompositing/Public/MovieSceneMediaData.h b/Engine/Plugins/Media/MediaCompositing/Source/MediaCompositing/Public/MovieSceneMediaData.h
--- a/Engine/Plugins/Media/MediaCompositing/Source/MediaCompositing/Public/MovieSceneMediaData.h
+++ b/Engine/Plugins/Media/MediaCompositing/Source/MediaCompositing/Public/MovieSceneMediaData.h
@@ -6,6 +6,7 @@
 #include "Evaluation/MovieScenePropertyTemplate.h"
 #include "UObject/WeakObjectPtrTemplates.h"
 
+class FMovieSceneMediaPlayerStore;
 class IMediaPlayerProxy;
 class UMediaPlayer;
 class UMediaTexture;
@@ -56,7 +57,7 @@ public:
 	void SeekOnOpen(FTimespan Time);
 
 	/** Set up this persistent data object. */
-	void Setup(UMediaPlayer* OverrideMediaPlayer, UObject* InPlayerProxy, int32 InProxyLayerIndex, int32 InProxyTextureIndex);
+	void Setup(const TSharedPtr<FMovieSceneMediaPlayerStore>& InMediaPlayerStore, const UObject* InSection, UMediaPlayer* OverrideMediaPlayer, UObject* InPlayerProxy, int32 InProxyLayerIndex, int32 InProxyTextureIndex);
 
 	/**
 	 * Called from FMovieSceneMediaSectionTemplate::Initialize.
@@ -103,4 +104,10 @@ private:
 
 	/** The time to seek to after the media source is opened. */
 	FTimespan SeekOnOpenTime;
+
+	/** Pointer to the media player store this section data was setup with. Will be used to return the media player to the store on destruction. */
+	TWeakPtr<FMovieSceneMediaPlayerStore> MediaPlayerStoreWeak;
+
+	/** Owning Media Section. Used as persistent identifier for media player data. */
+	const UObject* MediaSection = nullptr;
 };
