[MediaCompositing]
* Fixed non unity compile error

diff --git a/Engine/Plugins/Media/MediaCompositing/Source/MediaCompositing/Private/MovieScene/MovieSceneMediaPlayerStore.h b/Engine/Plugins/Media/MediaCompositing/Source/MediaCompositing/Private/MovieScene/MovieSceneMediaPlayerStore.h
--- a/Engine/Plugins/Media/MediaCompositing/Source/MediaCompositing/Private/MovieScene/MovieSceneMediaPlayerStore.h
+++ b/Engine/Plugins/Media/MediaCompositing/Source/MediaCompositing/Private/MovieScene/MovieSceneMediaPlayerStore.h
@@ -3,6 +3,7 @@
 #pragma once
 
 #include "Evaluation/PersistentEvaluationData.h"
+#include "UObject/ObjectKey.h"
 
 class UMediaPlayer;
 
