[MediaCompositing] Code clean up: replacting raw object pointer with a ObjetKey.

diff --git a/Engine/Plugins/Media/MediaCompositing/Source/MediaCompositing/Private/MovieScene/MovieSceneMediaPlayerStore.cpp b/Engine/Plugins/Media/MediaCompositing/Source/MediaCompositing/Private/MovieScene/MovieSceneMediaPlayerStore.cpp
--- a/Engine/Plugins/Media/MediaCompositing/Source/MediaCompositing/Private/MovieScene/MovieSceneMediaPlayerStore.cpp
+++ b/Engine/Plugins/Media/MediaCompositing/Source/MediaCompositing/Private/MovieScene/MovieSceneMediaPlayerStore.cpp
@@ -22,10 +22,10 @@ FMovieSceneMediaPlayerStore::~FMovieSceneMediaPlayerStore()
 	CloseRemainingPlayers();
 }
 
-void FMovieSceneMediaPlayerStore::ScheduleMediaPlayerForRelease(const UObject* InPersistentObject, UMediaPlayer* InMediaPlayer)
+void FMovieSceneMediaPlayerStore::ScheduleMediaPlayerForRelease(const FObjectKey& InPersistentObjectKey, UMediaPlayer* InMediaPlayer)
 {
 	// Validation code -- There shouldn't be more than one player per section in the current design.
-	if (UMediaPlayer** ExistingMediaPlayer = MediaPlayers.Find(InPersistentObject))
+	if (UMediaPlayer** ExistingMediaPlayer = MediaPlayers.Find(InPersistentObjectKey))
 	{
 		if (IsValid(*ExistingMediaPlayer) && *ExistingMediaPlayer != InMediaPlayer)
 		{
@@ -33,15 +33,15 @@ void FMovieSceneMediaPlayerStore::ScheduleMediaPlayerForRelease(const UObject* I
 		}
 	}
 	
-	MediaPlayers.Add(InPersistentObject, InMediaPlayer);
+	MediaPlayers.Add(InPersistentObjectKey, InMediaPlayer);
 }
 
-UMediaPlayer* FMovieSceneMediaPlayerStore::TryAcquireMediaPlayer(const UObject* InPersistentObject)
+UMediaPlayer* FMovieSceneMediaPlayerStore::TryAcquireMediaPlayer(const FObjectKey& InPersistentObjectKey)
 {
-	if (UMediaPlayer** FoundMediaPlayer = MediaPlayers.Find(InPersistentObject))
+	if (UMediaPlayer** FoundMediaPlayer = MediaPlayers.Find(InPersistentObjectKey))
 	{
 		UMediaPlayer* MediaPlayer = *FoundMediaPlayer;
-		MediaPlayers.Remove(InPersistentObject);
+		MediaPlayers.Remove(InPersistentObjectKey);
 		return MediaPlayer;
 	}
 	return nullptr;
diff --git a/Engine/Plugins/Media/MediaCompositing/Source/MediaCompositing/Private/MovieScene/MovieSceneMediaPlayerStore.h b/Engine/Plugins/Media/MediaCompositing/Source/MediaCompositing/Private/MovieScene/MovieSceneMediaPlayerStore.h
--- a/Engine/Plugins/Media/MediaCompositing/Source/MediaCompositing/Private/MovieScene/MovieSceneMediaPlayerStore.h
+++ b/Engine/Plugins/Media/MediaCompositing/Source/MediaCompositing/Private/MovieScene/MovieSceneMediaPlayerStore.h
@@ -28,17 +28,17 @@ public:
 	 * If not acquired again with the given persistent object by the end of the frame,
 	 * it will be closed and discarded.
 	 * 
-	 * @param InPersistentObject Persistent object (ex section) we want to associate the player with.
+	 * @param InPersistentObjectKey Persistent object (ex section) we want to associate the player with.
 	 * @param InMediaPlayer Media player to reuse.
 	 */
-	void ScheduleMediaPlayerForRelease(const UObject* InPersistentObject, UMediaPlayer* InMediaPlayer);
+	void ScheduleMediaPlayerForRelease(const FObjectKey& InPersistentObjectKey, UMediaPlayer* InMediaPlayer);
 
 	/**
 	 * @brief Get the media player registered to the given persistent object. Remove from pool if found.
-	 * @param InPersistentObject Persistent object (ex section) the media player is associated with.
+	 * @param InPersistentObjectKey Persistent object (ex section) the media player is associated with.
 	 * @return pointer to media player if found, null otherwise.
 	 */
-	UMediaPlayer* TryAcquireMediaPlayer(const UObject* InPersistentObject);
+	UMediaPlayer* TryAcquireMediaPlayer(const FObjectKey& InPersistentObjectKey);
 
 	/**
 	 * @brief Close the given media player and remove from root.
diff --git a/Engine/Plugins/Media/MediaCompositing/Source/MediaCompositing/Public/MovieSceneMediaData.h b/Engine/Plugins/Media/MediaCompositing/Source/MediaCompositing/Public/MovieSceneMediaData.h
--- a/Engine/Plugins/Media/MediaCompositing/Source/MediaCompositing/Public/MovieSceneMediaData.h
+++ b/Engine/Plugins/Media/MediaCompositing/Source/MediaCompositing/Public/MovieSceneMediaData.h
@@ -109,5 +109,5 @@ private:
 	TWeakPtr<FMovieSceneMediaPlayerStore> MediaPlayerStoreWeak;
 
 	/** Owning Media Section. Used as persistent identifier for media player data. */
-	const UObject* MediaSection = nullptr;
+	FObjectKey MediaSection;
 };
