[Media] Sequencer Media Track updates player's PlaybackRange.

The implementation is enabled with a cvar "MediaTrack.UpdatePlaybackRange", which is disabled by default because there are still issues using it with Electra. But it can be enabled when using Protron.

diff --git a/Engine/Plugins/Media/MediaCompositing/Source/MediaCompositing/Private/MovieScene/MovieSceneMediaData.cpp b/Engine/Plugins/Media/MediaCompositing/Source/MediaCompositing/Private/MovieScene/MovieSceneMediaData.cpp
--- a/Engine/Plugins/Media/MediaCompositing/Source/MediaCompositing/Private/MovieScene/MovieSceneMediaData.cpp
+++ b/Engine/Plugins/Media/MediaCompositing/Source/MediaCompositing/Private/MovieScene/MovieSceneMediaData.cpp
@@ -170,27 +170,54 @@ void FMovieSceneMediaData::HandleMediaPlayerEvent(EMediaEvent Event)
 	}
 
 	const FTimespan Duration = MediaPlayer->GetDuration();
-    if (Duration == FTimespan::Zero())
-    {
-	    return;
-    }
+	if (Duration == FTimespan::Zero())
+	{
+		return;
+	}
 
-	// Update looping here, as if there is no IMediaPlayer (which is different to MediaPlayer)
-	// when we last called SetLooping then looping might not be set correctly.
-	MediaPlayer->SetLooping(MediaPlayer->IsLooping());
+	// Set looping from the media section parameter.
+	// Remark: this must be set to the value it is going to be in the media section, switching it will cause a seek/flush on Electra.
+	MediaPlayer->SetLooping(bIsLooping);
 
 	FTimespan MediaTime;
-	if (!MediaPlayer->IsLooping())
+	
+	if (!PlaybackRange.IsEmpty() && MediaPlayer->SupportsPlaybackTimeRange())
 	{
-		MediaTime = FMath::Clamp(SeekOnOpenTime, FTimespan::Zero(), Duration - FTimespan::FromSeconds(0.001));
+		const TRange<FTimespan> CurrentPlaybackRange = MediaPlayer->GetPlaybackTimeRange(EMediaTimeRangeType::Current);
+		const TRange<FTimespan> FullPlaybackRange = MediaPlayer->GetPlaybackTimeRange(EMediaTimeRangeType::Absolute);
+		// Make sure the playback range from movie scene is valid.
+		const TRange<FTimespan> ClampedRange = TRange<FTimespan>::Intersection(PlaybackRange, FullPlaybackRange);
+		if (CurrentPlaybackRange != ClampedRange)
+		{
+			MediaPlayer->SetPlaybackTimeRange(ClampedRange);
+		}
+
+		// Ensure we are seeking within the specified playback range (very precisely).
+		if (!bIsLooping)
+		{
+			MediaTime = FMath::Clamp(SeekOnOpenTime, ClampedRange.GetLowerBoundValue(), ClampedRange.GetUpperBoundValue());
+		}
+		else
+		{
+			// Loop is defined within the playback range.
+			const FTimespan OffsetTime = FMath::Max(SeekOnOpenTime - ClampedRange.GetLowerBoundValue(), FTimespan::Zero()); 
+			MediaTime = (OffsetTime % ClampedRange.Size<FTimespan>()) + ClampedRange.GetLowerBoundValue();
+		}
 	}
 	else
 	{
-		MediaTime = SeekOnOpenTime % Duration;
+		if (!bIsLooping)
+		{
+			MediaTime = FMath::Clamp(SeekOnOpenTime, FTimespan::Zero(), Duration - FTimespan::FromSeconds(0.001));
+		}
+		else
+		{
+			MediaTime = SeekOnOpenTime % Duration;
+		}
 	}
 
 	MediaPlayer->SetRate(0.0f);
 	MediaPlayer->Seek(MediaTime);
-
+	
 	SeekOnOpenTime = FTimespan::MinValue();
 }
diff --git a/Engine/Plugins/Media/MediaCompositing/Source/MediaCompositing/Private/MovieScene/MovieSceneMediaTemplate.cpp b/Engine/Plugins/Media/MediaCompositing/Source/MediaCompositing/Private/MovieScene/MovieSceneMediaTemplate.cpp
--- a/Engine/Plugins/Media/MediaCompositing/Source/MediaCompositing/Private/MovieScene/MovieSceneMediaTemplate.cpp
+++ b/Engine/Plugins/Media/MediaCompositing/Source/MediaCompositing/Private/MovieScene/MovieSceneMediaTemplate.cpp
@@ -12,6 +12,7 @@
 #include "MediaSource.h"
 #include "MediaSourceOptions.h"
 #include "MediaTexture.h"
+#include "MovieScene.h"
 #include "UObject/Package.h"
 #include "UObject/GCObject.h"
 
@@ -31,6 +32,19 @@ static TAutoConsoleVariable<bool> CVarBlockOnFetchScrub(
 	TEXT("Enable/Disable BlockOnFetch while scrubbing."),
 	ECVF_Default);
 
+static TAutoConsoleVariable<bool> CVarUpdatePlaybackRange(
+	TEXT("MediaTrack.UpdatePlaybackRange"),
+	false, // Disabled by default until issues with Electra are ironed out.
+	TEXT("Update Player PlaybackRange. Improves looping performance with better pre-caching."),
+	ECVF_Default);
+
+static TAutoConsoleVariable<bool> CVarAvoidSeekOnJump(
+	TEXT("MediaTrack.AvoidSeekOnJump"),
+	true,
+	TEXT("When the playhead jumps (on loop for instance), seek only if needed, i.e. if player doesn't have cached samples already."),
+	ECVF_Default);
+
+
 /* Local helpers
  *****************************************************************************/
 
@@ -39,7 +53,7 @@ struct FMediaSectionBaseMediaSourceOptions
 	FMediaSourceCacheSettings CacheSettings;
 };
 
-namespace FMediaSectionBaseMediaUtils
+namespace UE::MediaSectionBaseMediaUtils
 {
 	FMediaSourceCacheSettings GetCurrentCacheSettingsFromPlayer(const UMediaPlayer* InPlayer)
 	{
@@ -68,10 +82,6 @@ namespace FMediaSectionBaseMediaUtils
 	 */
 	bool IsVideoSampleAvailableForTime(const TSharedRef<FMediaPlayerFacade>& InMediaPlayerFacade, const FTimespan& InTime)
 	{
-		// Todo: For Protron, this doesn't work on loops. Only works for preroll.
-		// Preroll is caught by the first test. But doesn't work on loops.
-		// For loops, the QueryCacheState should catch it, but Protron never has any "loaded" samples.
-		
 		// Check LastVideoSampleProcessedTimeRange to see if the sample has already been consumed.
 		const TRange<FMediaTimeStamp> LastVideoSampleTimeRange = InMediaPlayerFacade->GetLastProcessedVideoSampleTimeRange();
 		if (!LastVideoSampleTimeRange.IsEmpty())
@@ -107,6 +117,131 @@ namespace FMediaSectionBaseMediaUtils
 		// Otherwise, check if we have a sample available for the desired time, either in the sample queue or already processed.
 		return IsVideoSampleAvailableForTime(MediaPlayerFacade, InTime);
 	}
+
+	/**
+	 * Utility to convert a sequencer's (absolute) frame number to a player's timespan.
+	 */
+	FTimespan SequencerFrameToPlayerTime(const FMovieSceneMediaSectionParams& InMediaSectionParams, const FFrameRate& InSequencerFrameRate, const FFrameNumber& InSequencerFrame)
+	{
+		const FFrameNumber PlayerFrame = InSequencerFrame - InMediaSectionParams.SectionStartFrame + InMediaSectionParams.StartFrameOffset;
+		return FTimespan::FromSeconds(InSequencerFrameRate.AsSeconds(PlayerFrame));
+	}
+
+	/** Converts specified time to a range bound of given type (inclusive or exclusive). */
+	template <typename T>
+	TRangeBound<T> ToRangeBound(bool bInIsExclusive, const T& InTime)
+	{
+		return bInIsExclusive ? TRangeBound<T>::Exclusive(InTime) : TRangeBound<T>::Inclusive(InTime); 
+	}
+
+	/**
+	 * Converts the given sequencer frame number range in media player's time span range.
+	 * @param InMediaSectionParams Media Section information
+	 * @param InSequencerFrameRate Sequencer frame rate to convert frame number in time span
+	 * @param InFrameRange Frame Number range to convert
+	 */
+	TRange<FTimespan> SequencerFrameRangeToPlayerRange(const FMovieSceneMediaSectionParams& InMediaSectionParams, const FFrameRate& InSequencerFrameRate, const TRange<FFrameNumber>& InFrameRange)
+	{
+		const FTimespan LowerTime = SequencerFrameToPlayerTime(InMediaSectionParams, InSequencerFrameRate, InFrameRange.GetLowerBoundValue());
+		const FTimespan UpperTime = SequencerFrameToPlayerTime(InMediaSectionParams, InSequencerFrameRate, InFrameRange.GetUpperBoundValue());
+		return TRange<FTimespan>(ToRangeBound(InFrameRange.GetLowerBound().IsExclusive(), LowerTime), ToRangeBound(InFrameRange.GetUpperBound().IsExclusive(), UpperTime));
+	}
+
+	/**
+	 * Converts the given frame number range in all inclusive bounds.
+	 * Media Player requires ranges to be inclusive.
+	 */
+	TRange<FFrameNumber> ToAllInclusiveRange(const TRange<FFrameNumber>& InSequencerFrameRange)
+	{
+		TRangeBound<FFrameNumber> LowerBound = InSequencerFrameRange.GetLowerBound();
+		TRangeBound<FFrameNumber> UpperBound = InSequencerFrameRange.GetUpperBound();
+		if (LowerBound.IsExclusive())
+		{
+			LowerBound = TRangeBound<FFrameNumber>::Inclusive(LowerBound.GetValue() + 1);
+		}
+		if (UpperBound.IsExclusive())
+		{
+			UpperBound = TRangeBound<FFrameNumber>::Inclusive(UpperBound.GetValue() - 1);
+		}
+		return TRange<FFrameNumber>(LowerBound, UpperBound);
+	}
+	
+	/**
+	 * Calculates the intersection of the section's bounds and scene's playback range.
+	 */
+	TRange<FTimespan> CalculatePlayerPlaybackRange(const FMovieSceneMediaSectionParams& InMediaSectionParams, const UMovieSceneMediaSection* InMediaSection, const FFrameRate& InSequencerFrameRate)
+	{
+		// Todo: validate this. See GetTrueRange(). For consistency, should be: { lower <= time < upper }
+		const TRange<FFrameNumber> SectionFrameRange = TRange<FFrameNumber>(TRangeBound<FFrameNumber>::Inclusive(InMediaSectionParams.SectionStartFrame), TRangeBound<FFrameNumber>::Exclusive(InMediaSectionParams.SectionEndFrame));
+		const TRange<FTimespan> SectionTimeRange = SequencerFrameRangeToPlayerRange(InMediaSectionParams, InSequencerFrameRate, ToAllInclusiveRange(SectionFrameRange));
+
+		const UMovieSceneTrack* Track = InMediaSection->GetTypedOuter<UMovieSceneTrack>();
+		const UMovieScene* Scene = Track ? Track->GetTypedOuter<UMovieScene>() : nullptr;
+
+		if (Scene)
+		{
+			const TRange<FTimespan> SceneTimeRange = SequencerFrameRangeToPlayerRange(InMediaSectionParams, InSequencerFrameRate, ToAllInclusiveRange(Scene->GetPlaybackRange()));
+			return TRange<FTimespan>::Intersection(SectionTimeRange, SceneTimeRange);
+		}
+
+		return SectionTimeRange;
+	}
+
+	/** Utility function to set the specified playback time range */
+	void SetPlayerPlaybackTimeRange(UMediaPlayer* InMediaPlayer, const TRange<FTimespan>& InPlaybackRange)
+	{
+		if (InMediaPlayer->SupportsPlaybackTimeRange())
+		{
+			const TRange<FTimespan> CurrentPlaybackRange = InMediaPlayer->GetPlaybackTimeRange(EMediaTimeRangeType::Current);
+			const TRange<FTimespan> FullPlaybackRange = InMediaPlayer->GetPlaybackTimeRange(EMediaTimeRangeType::Absolute);
+
+			if (!InPlaybackRange.IsEmpty())
+			{
+				// Make sure the playback range from movie scene is valid.
+				const TRange<FTimespan> ClampedRange = TRange<FTimespan>::Intersection(InPlaybackRange, FullPlaybackRange);
+				
+				if (CurrentPlaybackRange != ClampedRange)
+				{
+					InMediaPlayer->SetPlaybackTimeRange(ClampedRange);
+				}
+			}
+			else
+			{
+				// Restore playback range to full clip length (default)
+				if (CurrentPlaybackRange != FullPlaybackRange)
+				{
+					InMediaPlayer->SetPlaybackTimeRange(FullPlaybackRange);
+				}
+			}
+		}
+	}
+
+	/**
+	 * Helper function to set BlockOnTimeRange.
+	 * If a PlaybackRange is specified, the block on range must be within the playback range,
+	 * very precisely to avoid a blocking time out.
+	 * TODO: this is probably best done in FMediaPlayerFacade. 
+	 */
+	void SetPlayerBlockOnTimeRange(UMediaPlayer* InMediaPlayer, const FTimespan& InCurrentTime, const FTimespan& InFrameDuration, const TRange<FTimespan>& InPlaybackRange)
+	{
+		FTimespan CurrentTime = InCurrentTime;
+		if (!InPlaybackRange.IsEmpty() && InMediaPlayer->SupportsPlaybackTimeRange())
+		{
+			// Try to ensure the specified time range is within the playback range very precisely. (otherwise it will block forever)
+			const TRange<FTimespan> CurrentPlaybackRange = InMediaPlayer->GetPlaybackTimeRange(EMediaTimeRangeType::Current);
+			if (!InMediaPlayer->IsLooping())
+			{
+				CurrentTime = FMath::Clamp(CurrentTime, CurrentPlaybackRange.GetLowerBoundValue(), CurrentPlaybackRange.GetUpperBoundValue() - InFrameDuration * 0.5);
+			}
+			else
+			{
+				// TODO
+				// ... Looping case is complex, need to match FMediaPlayerFacade::FBlockOnRange::GetRange() ...
+			}
+		}
+
+		InMediaPlayer->SetBlockOnTimeRange(TRange<FTimespan>(CurrentTime, CurrentTime + InFrameDuration));
+	}
 }
 
 /** Base struct for execution tokens. */
@@ -212,10 +347,11 @@ private:
 struct FMediaSectionExecutionToken
 	: FMediaSectionBaseExecutionToken
 {
-	FMediaSectionExecutionToken(UMediaSource* InMediaSource, const FMediaSectionBaseMediaSourceOptions& InMediaSourceOptions, FMovieSceneObjectBindingID InMediaSourceProxy, int32 InMediaSourceProxyIndex, float InProxyTextureBlend, bool bInCanPlayerBeOpen, FTimespan InCurrentTime, FTimespan InFrameDuration)
+	FMediaSectionExecutionToken(UMediaSource* InMediaSource, const FMediaSectionBaseMediaSourceOptions& InMediaSourceOptions, FMovieSceneObjectBindingID InMediaSourceProxy, int32 InMediaSourceProxyIndex, float InProxyTextureBlend, bool bInCanPlayerBeOpen, FTimespan InCurrentTime, FTimespan InFrameDuration, const TRange<FTimespan>& InPlaybackRange)
 		: FMediaSectionBaseExecutionToken(InMediaSource, InMediaSourceOptions, InMediaSourceProxy, InMediaSourceProxyIndex)
 		, CurrentTime(InCurrentTime)
 		, FrameDuration(InFrameDuration)
+		, PlaybackRange(InPlaybackRange)
 		, PlaybackRate(1.0f)
 		, ProxyTextureBlend(InProxyTextureBlend)
 		, bCanPlayerBeOpen(bInCanPlayerBeOpen)
@@ -275,10 +411,11 @@ struct FMediaSectionExecutionToken
 			return;
 		}
 
+		using namespace UE::MediaSectionBaseMediaUtils;
 
 		// Check if the cache options have changed mid playback.
 		//const FMediaSourceCacheSettings CurrentCacheSettings = GetBaseMediaSourceOptions().CacheSettings;
-		const FMediaSourceCacheSettings CurrentCacheSettings = FMediaSectionBaseMediaUtils::GetCurrentCacheSettingsFromPlayer(MediaPlayer);
+		const FMediaSourceCacheSettings CurrentCacheSettings = GetCurrentCacheSettingsFromPlayer(MediaPlayer);
 		const bool bCacheSettingsChanged = PlayerProxyInterface ? CurrentCacheSettings != PlayerProxyInterface->GetCacheSettings() : false;
 
 		// open the media source if necessary
@@ -297,7 +434,7 @@ struct FMediaSectionExecutionToken
 			}
 
 			// Setup an initial blocking range - MediaFramework will block (even through the opening process) in its next tick...
-			MediaPlayer->SetBlockOnTimeRange(TRange<FTimespan>(CurrentTime, CurrentTime + FrameDuration));
+			SetPlayerBlockOnTimeRange(MediaPlayer, CurrentTime, FrameDuration, PlaybackRange);
 			MediaPlayer->OpenSourceWithOptions(MediaSource, Options);
 			return;
 		}
@@ -307,8 +444,7 @@ struct FMediaSectionExecutionToken
 		if (MediaPlayer->IsPreparing())
 		{
 			SectionData.SeekOnOpen(CurrentTime);
-			MediaPlayer->SetBlockOnTimeRange(TRange<FTimespan>(CurrentTime, CurrentTime + FrameDuration));
-
+			SetPlayerBlockOnTimeRange(MediaPlayer, CurrentTime, FrameDuration, PlaybackRange);
 			return;
 		}
 
@@ -323,6 +459,8 @@ struct FMediaSectionExecutionToken
 		// update media player
 		//
 
+		SetPlayerPlaybackTimeRange(MediaPlayer, PlaybackRange);
+
 		// Setup media time (used for seeks)
 		FTimespan MediaTime = CurrentTime;
 
@@ -335,7 +473,7 @@ struct FMediaSectionExecutionToken
 			if (!MediaPlayer->IsPlaying())
 			{
 				// If the player has pre-rolled, it is ready and doesn't need to seek.
-				if (!FMediaSectionBaseMediaUtils::ShouldSeekForTime(MediaPlayer, MediaTime))
+				if (!UE::MediaSectionBaseMediaUtils::ShouldSeekForTime(MediaPlayer, MediaTime))
 				{
 					SeekPlayer(Context, MediaPlayer, MediaTime);
 				}
@@ -353,8 +491,8 @@ struct FMediaSectionExecutionToken
 			}
 			else
 			{
-				// Todo: avoid seek on loop. Causes loop index to be different between player and facade.
-				if (Context.HasJumped()) // && !FMediaSectionBaseMediaUtils::ShouldSeekForTime(MediaPlayer, MediaTime))
+				// Avoid seek on jump (loop) if the player already has cached samples.
+				if (Context.HasJumped() && (!CVarAvoidSeekOnJump.GetValueOnGameThread() || !ShouldSeekForTime(MediaPlayer, MediaTime)))
 				{
 					SeekPlayer(Context, MediaPlayer, MediaTime);
 				}
@@ -402,7 +540,7 @@ struct FMediaSectionExecutionToken
 		{
 			// Set blocking range / time-range to display
 			// (we always use the full current time for this, any adjustments to player timestamps are done internally)
-			MediaPlayer->SetBlockOnTimeRange(TRange<FTimespan>(CurrentTime, CurrentTime + FrameDuration));
+			SetPlayerBlockOnTimeRange(MediaPlayer, CurrentTime, FrameDuration, PlaybackRange);
 		}
 	}
 
@@ -410,6 +548,7 @@ private:
 
 	FTimespan CurrentTime;
 	FTimespan FrameDuration;
+	TRange<FTimespan> PlaybackRange;
 	float PlaybackRate;
 	float ProxyTextureBlend;
 	bool bCanPlayerBeOpen;
@@ -455,10 +594,13 @@ void FMovieSceneMediaSectionTemplate::Evaluate(const FMovieSceneEvaluationOperan
 	{
 		return;
 	}
-
+	
 	// @todo: account for video time dilation if/when these are added
 
-	if (Context.IsPreRoll())
+	bool bCanPlayerBeOpen = true;
+	MediaSection->ChannelCanPlayerBeOpen.Evaluate(Context.GetTime(), bCanPlayerBeOpen);
+
+	if (Context.IsPreRoll() && bCanPlayerBeOpen)
 	{
 		const FFrameRate FrameRate = Context.GetFrameRate();
 		const FFrameNumber StartFrame = Context.HasPreRollEndTime() ? Context.GetPreRollEndFrame() - Params.SectionStartFrame + Params.StartFrameOffset : Params.StartFrameOffset;
@@ -482,9 +624,13 @@ void FMovieSceneMediaSectionTemplate::Evaluate(const FMovieSceneEvaluationOperan
 		const double FrameDurationInSeconds = FMath::Max(FrameRate.AsSeconds(FFrameTime(1)), (Context.GetRange().Size<FFrameTime>()) / Context.GetFrameRate());
 		const int64 FrameDurationTicks = static_cast<int64>(FrameDurationInSeconds * ETimespan::TicksPerSecond);
 
-		float ProxyTextureBlend = MediaSection->EvaluateEasing(Context.GetTime());
-		bool bCanPlayerBeOpen = true;
-		MediaSection->ChannelCanPlayerBeOpen.Evaluate(Context.GetTime(), bCanPlayerBeOpen);
+		const float ProxyTextureBlend = MediaSection->EvaluateEasing(Context.GetTime());
+
+		TRange<FTimespan> PlayerPlaybackRange;
+		if (CVarUpdatePlaybackRange.GetValueOnGameThread())
+		{
+			PlayerPlaybackRange = UE::MediaSectionBaseMediaUtils::CalculatePlayerPlaybackRange(Params, MediaSection, FrameRate);
+		}
 
 		#if MOVIESCENEMEDIATEMPLATE_TRACE_EVALUATION
 			GLog->Logf(ELogVerbosity::Log, TEXT("Evaluating frame %i+%f, FrameRate %i/%i, FrameTicks %d, FrameDurationTicks %d"),
@@ -499,7 +645,7 @@ void FMovieSceneMediaSectionTemplate::Evaluate(const FMovieSceneEvaluationOperan
 
 		FMediaSectionBaseMediaSourceOptions Options;
 		Options.CacheSettings = Params.CacheSettings;
-		ExecutionTokens.Add(FMediaSectionExecutionToken(MediaSource, Options, Params.MediaSourceProxy, Params.MediaSourceProxyIndex, ProxyTextureBlend, bCanPlayerBeOpen, FTimespan(FrameTicks), FTimespan(FrameDurationTicks)));
+		ExecutionTokens.Add(FMediaSectionExecutionToken(MediaSource, Options, Params.MediaSourceProxy, Params.MediaSourceProxyIndex, ProxyTextureBlend, bCanPlayerBeOpen, FTimespan(FrameTicks), FTimespan(FrameDurationTicks), PlayerPlaybackRange));
 	}
 }
 
@@ -552,6 +698,13 @@ void FMovieSceneMediaSectionTemplate::Initialize(const FMovieSceneEvaluationOper
 		SectionData = &PersistentData.AddSectionData<FMovieSceneMediaData>();
 		const TSharedPtr<FMovieSceneMediaPlayerStore> MediaPlayerStore = FMovieSceneMediaPlayerStoreContainer::GetOrAdd(PersistentData).GetMediaPlayerStore();
 		SectionData->Setup(MediaPlayerStore, MediaSection, MediaPlayer, PlayerProxy, ProxyLayerIndex, ProxyTextureIndex);
+
+		// Prepare these values in the section data to have them set in the player when it is opened (as early as possible before the first seek).
+		SectionData->bIsLooping = Params.bLooping;
+		if (CVarUpdatePlaybackRange.GetValueOnGameThread())
+		{
+			SectionData->PlaybackRange = UE::MediaSectionBaseMediaUtils::CalculatePlayerPlaybackRange(Params, MediaSection, Context.GetFrameRate());
+		}
 	}
 
 	if (!ensure(SectionData != nullptr))
@@ -609,7 +762,10 @@ void FMovieSceneMediaSectionTemplate::Initialize(const FMovieSceneEvaluationOper
 		}
 	}
 
-	MediaPlayer->SetLooping(Params.bLooping);
+	if (MediaPlayer->IsLooping() != Params.bLooping)
+	{
+		MediaPlayer->SetLooping(Params.bLooping);
+	}
 }
 
 
diff --git a/Engine/Plugins/Media/MediaCompositing/Source/MediaCompositing/Public/MovieSceneMediaData.h b/Engine/Plugins/Media/MediaCompositing/Source/MediaCompositing/Public/MovieSceneMediaData.h
--- a/Engine/Plugins/Media/MediaCompositing/Source/MediaCompositing/Public/MovieSceneMediaData.h
+++ b/Engine/Plugins/Media/MediaCompositing/Source/MediaCompositing/Public/MovieSceneMediaData.h
@@ -74,6 +74,18 @@ public:
 	 */
 	bool bIsAspectRatioSet;
 
+	/**
+	 * Indicate if player is looping from the corresponding media section parameter.
+	 * We need to store this here to have it in HandleMediaPlayerEvent.
+	 */
+	bool bIsLooping = false;
+
+	/**
+	 * If specified, this is the playback range calculated from the section and sequencer.
+	 * We need to store this here to have it in HandleMediaPlayerEvent.
+	 */
+	TRange<FTimespan> PlaybackRange;
+
 private:
 	/**
 	 * Does the work needed so we can use our proxy media texture.
