MediaTrack: Exposed "Synchronous Scrubbing" option, disabled by default.

* Force synchronous frame requests while manually scrubbing the media track, at the cost of blocking the editor. Valuable for maintaining CG & media alignment.
* Note: Only respected by image media currently.

diff --git a/Engine/Plugins/Media/MediaCompositing/Source/MediaCompositing/Private/MovieScene/MovieSceneMediaTemplate.cpp b/Engine/Plugins/Media/MediaCompositing/Source/MediaCompositing/Private/MovieScene/MovieSceneMediaTemplate.cpp
--- a/Engine/Plugins/Media/MediaCompositing/Source/MediaCompositing/Private/MovieScene/MovieSceneMediaTemplate.cpp
+++ b/Engine/Plugins/Media/MediaCompositing/Source/MediaCompositing/Private/MovieScene/MovieSceneMediaTemplate.cpp
@@ -26,12 +26,6 @@
 
 #define MOVIESCENEMEDIATEMPLATE_TRACE_EVALUATION 0
 
-static TAutoConsoleVariable<bool> CVarBlockOnFetchScrub(
-	TEXT("MediaTrack.BlockOnFetchScrub"),
-	false,	// Should be disabled by default, but can be re-enabled in case of undesirable side effects in some use cases.
-	TEXT("Enable/Disable BlockOnFetch while scrubbing."),
-	ECVF_Default);
-
 static TAutoConsoleVariable<bool> CVarUpdatePlaybackRange(
 	TEXT("MediaTrack.UpdatePlaybackRange"),
 	false, // Disabled by default until issues with Electra are ironed out.
@@ -51,6 +45,7 @@ static TAutoConsoleVariable<bool> CVarAvoidSeekOnJump(
 struct FMediaSectionBaseMediaSourceOptions
 {
 	FMediaSourceCacheSettings CacheSettings;
+	bool bSynchronousScrubbing = false;
 };
 
 namespace UE::MediaSectionBaseMediaUtils
@@ -531,7 +526,7 @@ struct FMediaSectionExecutionToken
 		}
 
 		if ((Context.GetStatus() == EMovieScenePlayerStatus::Scrubbing || Context.GetStatus() == EMovieScenePlayerStatus::Stopped)
-			&& CVarBlockOnFetchScrub.GetValueOnGameThread() == false)
+			&& GetBaseMediaSourceOptions().bSynchronousScrubbing == false)
 		{
 			// When scrubbing, seek requests are non-blocking.
 			MediaPlayer->SetBlockOnTimeRange(TRange<FTimespan>::Empty());
@@ -610,6 +605,7 @@ void FMovieSceneMediaSectionTemplate::Evaluate(const FMovieSceneEvaluationOperan
 
 		FMediaSectionBaseMediaSourceOptions Options;
 		Options.CacheSettings = Params.CacheSettings;
+		Options.bSynchronousScrubbing = false; // not relevant during pre-roll
 		ExecutionTokens.Add(FMediaSectionPreRollExecutionToken(MediaSource, Options, Params.MediaSourceProxy, Params.MediaSourceProxyIndex, FTimespan(StartTicks)));
 	}
 	else if (!Context.IsPostRoll() && (Context.GetTime().FrameNumber < Params.SectionEndFrame))
@@ -643,8 +639,11 @@ void FMovieSceneMediaSectionTemplate::Evaluate(const FMovieSceneEvaluationOperan
 			);
 		#endif
 
+		const UMovieSceneMediaTrack* MediaTrack = MediaSection->GetTypedOuter<UMovieSceneMediaTrack>();
+
 		FMediaSectionBaseMediaSourceOptions Options;
 		Options.CacheSettings = Params.CacheSettings;
+		Options.bSynchronousScrubbing = MediaTrack ? MediaTrack->bSynchronousScrubbing : false;
 		ExecutionTokens.Add(FMediaSectionExecutionToken(MediaSource, Options, Params.MediaSourceProxy, Params.MediaSourceProxyIndex, ProxyTextureBlend, bCanPlayerBeOpen, FTimespan(FrameTicks), FTimespan(FrameDurationTicks), PlayerPlaybackRange));
 	}
 }
diff --git a/Engine/Plugins/Media/MediaCompositing/Source/MediaCompositing/Public/MovieSceneMediaTrack.h b/Engine/Plugins/Media/MediaCompositing/Source/MediaCompositing/Public/MovieSceneMediaTrack.h
--- a/Engine/Plugins/Media/MediaCompositing/Source/MediaCompositing/Public/MovieSceneMediaTrack.h
+++ b/Engine/Plugins/Media/MediaCompositing/Source/MediaCompositing/Public/MovieSceneMediaTrack.h
@@ -80,6 +80,15 @@ public:
 
 #endif // WITH_EDITOR
 
+	/**
+	 * Force synchronous frame requests *while* manually scrubbing the media track, at the cost of blocking the editor.
+	 * This can be valuable for always maintaining alignment of media & CG while working.
+	 * 
+	 * Note: Only respected by image sequences currently.
+	 */
+	UPROPERTY(EditAnywhere, Category = "Media Track")
+	bool bSynchronousScrubbing = false;
+
 private:
 	/** Base function to add a new section. */
 	UMovieSceneMediaSection* AddNewSectionOnRow(FFrameNumber Time, int32 RowIndex);
