[MediaCompositing] Helping scrubbing frame accuracy by ensuring the specified tick is in the correct frame.

When converting from frame number to a time span (in integer ticks), neither flooring nor rounding the fractional tick values garantees the integer tick value is in the right frame. Only Ceil will garantee that. Ceil will either be the first or second tick of the desired frame. It should be sufficient to garantee the correctness of the frame for the player regardless of it's internal rounding.

diff --git a/Engine/Plugins/Media/MediaCompositing/Source/MediaCompositing/Private/MovieScene/MovieSceneMediaTemplate.cpp b/Engine/Plugins/Media/MediaCompositing/Source/MediaCompositing/Private/MovieScene/MovieSceneMediaTemplate.cpp
--- a/Engine/Plugins/Media/MediaCompositing/Source/MediaCompositing/Private/MovieScene/MovieSceneMediaTemplate.cpp
+++ b/Engine/Plugins/Media/MediaCompositing/Source/MediaCompositing/Private/MovieScene/MovieSceneMediaTemplate.cpp
@@ -599,9 +599,8 @@ void FMovieSceneMediaSectionTemplate::Evaluate(const FMovieSceneEvaluationOperan
 	{
 		const FFrameRate FrameRate = Context.GetFrameRate();
 		const FFrameNumber StartFrame = Context.HasPreRollEndTime() ? Context.GetPreRollEndFrame() - Params.SectionStartFrame + Params.StartFrameOffset : Params.StartFrameOffset;
-
-		const double StartFrameInSeconds = FrameRate.AsSeconds(StartFrame);
-		const int64 StartTicks = static_cast<int64>(StartFrameInSeconds * ETimespan::TicksPerSecond);
+		// Using Ceil to ensure the tick value is part of the desired frame, and not the frame before.
+		const int64 StartTicks = FMath::CeilToInt64(FrameRate.AsSeconds(StartFrame) * ETimespan::TicksPerSecond);
 
 		FMediaSectionBaseMediaSourceOptions Options;
 		Options.CacheSettings = Params.CacheSettings;
@@ -612,9 +611,8 @@ void FMovieSceneMediaSectionTemplate::Evaluate(const FMovieSceneEvaluationOperan
 	{
 		const FFrameRate FrameRate = Context.GetFrameRate();
 		const FFrameTime FrameTime(Context.GetTime().FrameNumber - Params.SectionStartFrame + Params.StartFrameOffset);
-
-		const double FrameTimeInSeconds = FrameRate.AsSeconds(FrameTime);
-		const int64 FrameTicks = static_cast<int64>(FrameTimeInSeconds * ETimespan::TicksPerSecond);
+		// Using Ceil to ensure the tick value is part of the desired frame, and not the frame before.
+		const int64 FrameTicks = FMath::CeilToInt64(FrameRate.AsSeconds(FrameTime) * ETimespan::TicksPerSecond);
 
 		// With zero-length frames (which can occur occasionally), we use the fixed frame time, matching previous behavior.
 		const double FrameDurationInSeconds = FMath::Max(FrameRate.AsSeconds(FFrameTime(1)), (Context.GetRange().Size<FFrameTime>()) / Context.GetFrameRate());
