[MediaCompositing] Media IO Sequencer: Display Media Info + GOP Codec Warning

diff --git a/Engine/Plugins/Media/MediaCompositing/Source/MediaCompositingEditor/Private/Sequencer/MediaThumbnailSection.cpp b/Engine/Plugins/Media/MediaCompositing/Source/MediaCompositingEditor/Private/Sequencer/MediaThumbnailSection.cpp
--- a/Engine/Plugins/Media/MediaCompositing/Source/MediaCompositingEditor/Private/Sequencer/MediaThumbnailSection.cpp
+++ b/Engine/Plugins/Media/MediaCompositing/Source/MediaCompositingEditor/Private/Sequencer/MediaThumbnailSection.cpp
@@ -12,7 +12,6 @@
 #include "MediaPlayer.h"
 #include "MediaPlayerFacade.h"
 #include "MediaSource.h"
-#include "MediaTexture.h"
 #include "Modules/ModuleManager.h"
 #include "MovieScene.h"
 #include "MovieSceneMediaSection.h"
@@ -32,6 +31,91 @@
 
 #define LOCTEXT_NAMESPACE "FMediaThumbnailSection"
 
+namespace UE::MediaThumbnailSection
+{
+	constexpr float FilmBorderHeight = 9.0f;
+	static const FLazyName SectionTitleFont("NormalFont");
+	static const FLazyName MediaInfoFont("SmallFont");
+
+	enum class EOffsetOrigin : int8
+	{
+		FromTop,
+		FromBottom
+	};
+
+	/** Helper to paint text for media info. */
+	struct FTextPaintHelper
+	{
+		FTextPaintHelper(const FMargin& InContentPadding, const ESlateDrawEffect InDrawEffects, const FSlateFontInfo& InFontInfo)
+			: ContentPadding(InContentPadding)
+			, DrawEffects(InDrawEffects)
+			, FontInfo(InFontInfo)
+			, FontMeasureService(FSlateApplication::Get().GetRenderer()->GetFontMeasureService())
+		{}
+
+		/**
+		 * Paint given string
+		 * @param InPainter Painter
+		 * @param InString String
+		 * @param InVerticalOffset Vertical offset relative to either the top or bottom of the section.
+		 * @param InOffsetOrigin Origin of the vertical offset
+		 * @return 
+		 */
+		FVector2f PaintString(FSequencerSectionPainter& InPainter, const FString& InString, float InVerticalOffset, EOffsetOrigin InOffsetOrigin)
+		{
+			FColor CurrentColor = TextColor;
+			CurrentColor.A = static_cast<uint8>(InPainter.GhostAlpha * 255);
+			FVector2f TextSize = FontMeasureService->Measure(InString, FontInfo);
+			FVector2D TextOffset(EForceInit::ForceInitToZero);
+			
+			// Have the text on the bottom left side of the clip rect, along with single thumbnail and section title.
+			if (InOffsetOrigin == EOffsetOrigin::FromBottom)
+			{
+				FVector2D BottomLeft = InPainter.SectionGeometry.AbsoluteToLocal(InPainter.SectionClippingRect.GetBottomLeft());
+				TextOffset.Set(BottomLeft.X + ContentPadding.Left + 2, InPainter.SectionGeometry.Size.Y - (TextSize.Y + ContentPadding.Bottom));
+				TextOffset.Y += InVerticalOffset;
+			}
+			else
+			{
+				FVector2D TopLeft = InPainter.SectionGeometry.AbsoluteToLocal(InPainter.SectionClippingRect.GetTopLeft());
+				TextOffset.Set(TopLeft.X + ContentPadding.Left + 2, 0);
+				TextOffset.Y += InVerticalOffset;
+			}
+			
+			int32 LayerId = InPainter.LayerId++;
+
+			// Drop shadow
+			FSlateDrawElement::MakeText(
+				InPainter.DrawElements,
+				LayerId,
+				InPainter.SectionGeometry.ToPaintGeometry(TextSize, FSlateLayoutTransform(TextOffset + FVector2D(1.f, 1.f))),
+				InString,
+				FontInfo,
+				DrawEffects,
+				FLinearColor(0, 0, 0, .5f * InPainter.GhostAlpha)
+			);
+			
+			FSlateDrawElement::MakeText(
+				InPainter.DrawElements,
+				LayerId,
+				InPainter.SectionGeometry.ToPaintGeometry(TextSize, FSlateLayoutTransform(TextOffset)),
+				InString,
+				FontInfo,
+				DrawEffects,
+				CurrentColor
+			);
+
+			return TextSize;	
+		}
+
+		FMargin ContentPadding;
+		const ESlateDrawEffect DrawEffects;
+		const FSlateFontInfo FontInfo;
+		const TSharedRef<FSlateFontMeasure> FontMeasureService;
+
+		FColor TextColor = FColor(192, 192, 192);
+	};
+}
 
 /* FMediaThumbnailSection structors
  *****************************************************************************/
@@ -41,11 +125,22 @@ FMediaThumbnailSection::FMediaThumbnailSection(UMovieSceneMediaSection& InSectio
 	, SectionPtr(&InSection)
 {
 	TimeSpace = ETimeSpace::Local;
+
+	if (InSequencer)
+	{
+		InSequencer->OnBeginScrubbingEvent().AddRaw(this, &FMediaThumbnailSection::OnBeginScrubbingEvent);
+		InSequencer->OnEndScrubbingEvent().AddRaw(this, &FMediaThumbnailSection::OnEndScrubbingEvent);
+	}
 }
 
 
 FMediaThumbnailSection::~FMediaThumbnailSection()
 {
+	if (const TSharedPtr<ISequencer> Sequencer = SequencerPtr.Pin())
+	{
+		Sequencer->OnBeginScrubbingEvent().RemoveAll(this);
+		Sequencer->OnEndScrubbingEvent().RemoveAll(this);
+	}
 }
 
 
@@ -68,9 +163,26 @@ FMargin FMediaThumbnailSection::GetContentPadding() const
 
 float FMediaThumbnailSection::GetSectionHeight(const UE::Sequencer::FViewDensityInfo& ViewDensity) const
 {
-	return FThumbnailSection::GetSectionHeight(ViewDensity) + 2 * 9.0f; // make space for the film border
-}
+	using namespace UE::MediaThumbnailSection;
 
+	// Calculate the base section height
+	float SectionHeight = FThumbnailSection::GetSectionHeight(ViewDensity);
+
+	// Calculate the section title height.
+	const float SectionTitleHeight = FAppStyle::GetFontStyle(SectionTitleFont).Size + 8.f;
+
+	// Calculate minimum space for the section title, media info and optional performance warning.
+	const float PlayerInfoHeight = FAppStyle::GetFontStyle(MediaInfoFont).Size + 8.f;
+	const float NumInfoLines = bDrawSeekPerformanceWarning ? 2.0f : 1.0f;
+	const float MinimumHeight = SectionTitleHeight + PlayerInfoHeight * NumInfoLines + 8.0f;
+
+	// Base section height is either the thumbnail height or just the title.
+	// We want to make sure we have enough space for the media info too.
+	SectionHeight = FMath::Max(SectionHeight, MinimumHeight);
+
+	// Make additional space for the film border
+	return SectionHeight + 2 * FilmBorderHeight;
+}
 
 FText FMediaThumbnailSection::GetSectionTitle() const
 {
@@ -179,6 +291,11 @@ void FMediaThumbnailSection::Tick(const FGeometry& AllottedGeometry, const FGeom
 		}
 	}
 
+	if (const UMediaPlayer* MediaPlayer = GetTemplateMediaPlayer())
+	{
+		UpdateCachedMediaInfo(MediaPlayer);
+	}
+	
 	FThumbnailSection::Tick(AllottedGeometry, ClippedGeometry, InCurrentTime, InDeltaTime);
 }
 
@@ -383,74 +500,48 @@ void FMediaThumbnailSection::DrawSampleStates(FSequencerSectionPainter& InPainte
 	}
 }
 
-
-void FMediaThumbnailSection::DrawMediaInfo(FSequencerSectionPainter& InPainter, 
-	UMediaPlayer* MediaPlayer,FVector2D SectionSize) const
+void FMediaThumbnailSection::DrawMediaInfo(FSequencerSectionPainter& InPainter,
+	UMediaPlayer* MediaPlayer, FVector2D SectionSize) const
 {
-	// Get tile info.
-	FString TileString;
-	FIntPoint TileNum(EForceInit::ForceInitToZero);
-	MediaPlayer->GetMediaInfo<FIntPoint>(TileNum, UMediaPlayer::MediaInfoNameSourceNumTiles.Resolve());
-	int32 TileTotalNum = TileNum.X * TileNum.Y;
-	if (TileTotalNum > 1)
+	using namespace UE::MediaThumbnailSection;
+
+	FTextPaintHelper TextPaintHelper(
+		GetContentPadding(),
+		InPainter.bParentEnabled ? ESlateDrawEffect::None : ESlateDrawEffect::DisabledEffect,
+		FAppStyle::GetFontStyle(MediaInfoFont));
+
+	// We have 2 lines of text to render, top one is the player info
+	// that we want to render below the section title.
+	// Second line is the warning message that we want to render at the bottom of the section.
+
+	// The available height may be less than what we need, some derived classes override the section height.
+	const float AvailableSectionHeight = InPainter.SectionGeometry.Size.Y - FilmBorderHeight;
+
+	float FontHeight = TextPaintHelper.FontInfo.Size + 4.f;
+
+	// Calculate the offset from the top for the player info.
+	float TextYOffsetFromTop = FilmBorderHeight + 4.f;
+
+	// Check if there is a section title (there may not be)
+	if (!GetSectionTitle().IsEmpty())
 	{
-		TileString = FText::Format(LOCTEXT("TileNum", "Tiles: {0}"), TileTotalNum).ToString();
+		TextYOffsetFromTop += FAppStyle::GetFontStyle(SectionTitleFont).Size + 8.f;
 	}
 
-	// Get mip info.
-	FString MipString;
-	int32 MipNum = 0;
-	MediaPlayer->GetMediaInfo<int32>(MipNum, UMediaPlayer::MediaInfoNameSourceNumMips.Resolve());
-	if (MipNum > 1)
+	// Ensure we have enough room to render
+	if (AvailableSectionHeight >= TextYOffsetFromTop + FontHeight)
 	{
-		MipString = FText::Format(LOCTEXT("Mips", "Mips: {0}"), MipNum).ToString();
-	}
+		FVector2f TextSize = TextPaintHelper.PaintString(InPainter, PlayerInfo, TextYOffsetFromTop, EOffsetOrigin::FromTop);
 
-	const ESlateDrawEffect DrawEffects = InPainter.bParentEnabled ? ESlateDrawEffect::None : ESlateDrawEffect::DisabledEffect;
-	const FSlateFontInfo SmallLayoutFont = FCoreStyle::GetDefaultFontStyle("Bold", 10);
-	const TSharedRef< FSlateFontMeasure > FontMeasureService = FSlateApplication::Get().GetRenderer()->GetFontMeasureService();
-	
-	FVector2f TextSize(EForceInit::ForceInitToZero);
-	FMargin ContentPadding = GetContentPadding();
-	FVector2D TextOffset(EForceInit::ForceInitToZero);
-	float BaseYOffset = 0.0f;
-
-	// Add tile string.
-	if (TileString.Len() > 0)
-	{
-		TextSize = FontMeasureService->Measure(TileString, SmallLayoutFont);
-		TextOffset.Set(ContentPadding.Left + 4,
-			InPainter.SectionGeometry.Size.Y - (TextSize.Y + ContentPadding.Bottom));
-		TextOffset.Y += BaseYOffset;
-		FSlateDrawElement::MakeText(
-			InPainter.DrawElements,
-			InPainter.LayerId++,
-			InPainter.SectionGeometry.ToPaintGeometry(TextSize, FSlateLayoutTransform(TextOffset)),
-			TileString,
-			SmallLayoutFont,
-			DrawEffects,
-			FColor(192, 192, 192, static_cast<uint8>(InPainter.GhostAlpha * 255))
-		);
-		BaseYOffset -= TextSize.Y + 2.0f;
-	}
+		float RemainingSize = AvailableSectionHeight - (TextYOffsetFromTop + TextSize.Y);
 
-	// Add mips string.
-	if (MipString.Len() > 0)
-	{
-		TextSize = FontMeasureService->Measure(MipString, SmallLayoutFont);
-		TextOffset.Set(ContentPadding.Left + 4,
-			InPainter.SectionGeometry.Size.Y - (TextSize.Y + ContentPadding.Bottom));
-		TextOffset.Y += BaseYOffset;
-		FSlateDrawElement::MakeText(
-			InPainter.DrawElements,
-			InPainter.LayerId++,
-			InPainter.SectionGeometry.ToPaintGeometry(TextSize, FSlateLayoutTransform(TextOffset)),
-			MipString,
-			SmallLayoutFont,
-			DrawEffects,
-			FColor(192, 192, 192, static_cast<uint8>(InPainter.GhostAlpha * 255))
-		);
-		BaseYOffset -= TextSize.Y + 2.0f;
+		// Avoid having both messages overlapping.
+		if (bDrawSeekPerformanceWarning && RemainingSize > FontHeight)
+		{
+			TGuardValue PushColor(TextPaintHelper.TextColor, FColor::Yellow);
+			FString SeekPerformanceWarning = LOCTEXT("SeekPerformance", "Slow Seek Performance (GOP codec)").ToString();
+			TextPaintHelper.PaintString(InPainter, SeekPerformanceWarning, 0.0f, EOffsetOrigin::FromBottom);
+		}
 	}
 }
 
@@ -604,4 +695,101 @@ void FMediaThumbnailSection::CopyTexture(FSlateTextureRenderTarget2DResource* Re
 		});
 }
 
+void FMediaThumbnailSection::OnBeginScrubbingEvent()
+{
+	bIsSequencerScrubbing = true;
+}
+
+void FMediaThumbnailSection::OnEndScrubbingEvent()
+{
+	bIsSequencerScrubbing = false;
+}
+
+void FMediaThumbnailSection::UpdateCachedMediaInfo(const UMediaPlayer* InMediaPlayer)
+{
+	// Note: Protron IsPaused() returns false even if Rate is zero
+	const bool bPlayerIsPaused = InMediaPlayer->GetRate() == 0.0f; 
+	
+	// Only check if player is paused
+	if (bPlayerIsPaused)
+	{
+		// Only check for seek performance warning if the sequencer is scrubbing.
+		if (bIsSequencerScrubbing)
+		{
+			int32 KeyframeInterval = -1;
+			// Note: This is only supported by Protron currently.
+			if (InMediaPlayer->GetMediaInfo<int32>(KeyframeInterval, UMediaPlayer::MediaInfoNameKeyframeInterval.Resolve()))
+			{
+				// The keyframe interval is one of:
+				//  -1 : no information returned (unknown)
+				//   0 : unknown keyframe spacing, not every frame is a keyframe but the spacing is variable or cannot be determined
+				//   1 : every frame is a keyframe
+				//  >1 : every n'th frame is a keyframe
+				if (KeyframeInterval != 1)
+				{
+					bDrawSeekPerformanceWarning = true;
+				}
+			}
+		}
+	}
+	else
+	{
+		bDrawSeekPerformanceWarning = false;
+	}
+
+	// Build the player information string.
+	PlayerInfo.Reset(512);
+
+	// Full player info only when paused. We don't want to overload on playback.
+	if (bPlayerIsPaused)
+	{
+		PlayerInfo = InMediaPlayer->GetPlayerName().ToString();
+
+		constexpr int32 SelectedTrackIndex = INDEX_NONE;
+		constexpr int32 SelectedFormatIndex = INDEX_NONE;
+	
+		FString Format = InMediaPlayer->GetVideoTrackType(SelectedTrackIndex, SelectedFormatIndex);
+		if (!Format.IsEmpty())
+		{
+			PlayerInfo.Appendf(TEXT("%s%s"), PlayerInfo.IsEmpty() ? TEXT("") : TEXT(", "), *Format);
+		}
+
+		FIntPoint Resolution = InMediaPlayer->GetVideoTrackDimensions(SelectedTrackIndex, SelectedFormatIndex);
+		if (Resolution.X != 0 || Resolution.Y != 0)
+		{
+			PlayerInfo.Appendf(TEXT("%s%dx%d"), PlayerInfo.IsEmpty() ? TEXT("") : TEXT(", "), Resolution.X, Resolution.Y);
+		}
+
+		float FrameRate = InMediaPlayer->GetVideoTrackFrameRate(SelectedTrackIndex, SelectedFormatIndex);
+		if (FrameRate > 0.0f)
+		{
+			PlayerInfo.Appendf(TEXT("%s%s"), PlayerInfo.IsEmpty() ? TEXT("") : TEXT(", "),
+				*FText::Format(LOCTEXT("FrameRateUnits", "{0} fps"), FrameRate).ToString());
+		}
+	}
+
+	// Get tile info. Note: this was previously displayed while running so it is kept that way.
+	FIntPoint TileNum(EForceInit::ForceInitToZero);
+	if (InMediaPlayer->GetMediaInfo<FIntPoint>(TileNum, UMediaPlayer::MediaInfoNameSourceNumTiles.Resolve()))
+	{
+		int32 TileTotalNum = TileNum.X * TileNum.Y;
+		if (TileTotalNum > 1)
+		{
+			PlayerInfo.Appendf(TEXT("%s%s"), PlayerInfo.IsEmpty() ? TEXT("") : TEXT(", "),
+				*FText::Format(LOCTEXT("TileNum", "Tiles: {0}"), TileTotalNum).ToString());
+		}
+	}
+
+	// Get mip info. Note: this was previously displayed while running so it is kept that way.
+	int32 MipNum = 0;
+	if (InMediaPlayer->GetMediaInfo<int32>(MipNum, UMediaPlayer::MediaInfoNameSourceNumMips.Resolve()))
+	{
+		if (MipNum > 1)
+		{
+			PlayerInfo.Appendf(TEXT("%s%s"), PlayerInfo.IsEmpty() ? TEXT("") : TEXT(", "),
+				*FText::Format(LOCTEXT("Mips", "Mips: {0}"), MipNum).ToString());
+		}
+	}
+}
+
 #undef LOCTEXT_NAMESPACE
diff --git a/Engine/Plugins/Media/MediaCompositing/Source/MediaCompositingEditor/Public/Sequencer/MediaThumbnailSection.h b/Engine/Plugins/Media/MediaCompositing/Source/MediaCompositingEditor/Public/Sequencer/MediaThumbnailSection.h
--- a/Engine/Plugins/Media/MediaCompositing/Source/MediaCompositingEditor/Public/Sequencer/MediaThumbnailSection.h
+++ b/Engine/Plugins/Media/MediaCompositing/Source/MediaCompositingEditor/Public/Sequencer/MediaThumbnailSection.h
@@ -124,7 +124,15 @@ protected:
 	void CopyTexture(FSlateTextureRenderTarget2DResource* RenderTarget, FTextureReferenceRHIRef SourceTexture);
 
 private:
+	/** Called when sequencer begins scrubbing. */
+	void OnBeginScrubbingEvent();
 
+	/** Called when sequencer ends scrubbing. */
+	void OnEndScrubbingEvent();
+
+	/** Update cached information from the player. (called every tick) */
+	void UpdateCachedMediaInfo(const UMediaPlayer* InMediaPlayer);
+	
 	/** The section object that owns this section. */
 	TWeakObjectPtr<UMovieSceneMediaSection> SectionPtr;
 
@@ -133,4 +141,13 @@ private:
 
 	/** Cached start time valid only during resize */
 	FFrameNumber InitialStartTimeDuringResize;
+
+	/** Cached player info string to be displayed in the next paint. */
+	FString PlayerInfo;
+
+	/** Show the seek performance warning. */
+	bool bDrawSeekPerformanceWarning = false;
+
+	/** True if the sequencer is currently scrubbing. */
+	bool bIsSequencerScrubbing = false;
 };
