Electra: Re-enable BlockOnFetch on both playing and paused states.

* The previous removal of paused state blocking [at] CL 34527816 is no longer needed given recent improvements to scrubbing on the sequencer side. Blocking is now automatically disabled by an empty block time range, see use of `bSynchronousScrubbing`.
* This allows electra players to also respect the new `bSynchronousScrubbing` media track checkbox.

diff --git a/Engine/Plugins/Media/ElectraPlayer/Source/ElectraPlayerPlugin/Private/ElectraPlayerPlugin.cpp b/Engine/Plugins/Media/ElectraPlayer/Source/ElectraPlayerPlugin/Private/ElectraPlayerPlugin.cpp
--- a/Engine/Plugins/Media/ElectraPlayer/Source/ElectraPlayerPlugin/Private/ElectraPlayerPlugin.cpp
+++ b/Engine/Plugins/Media/ElectraPlayer/Source/ElectraPlayerPlugin/Private/ElectraPlayerPlugin.cpp
@@ -864,7 +864,7 @@ bool FElectraPlayerPlugin::CanControl(EMediaControl Control) const
 	EMediaState CurrentState = GetState();
 	if (Control == EMediaControl::BlockOnFetch)
 	{
-		return CurrentState == EMediaState::Playing;
+		return CurrentState == EMediaState::Playing || CurrentState == EMediaState::Paused;
 	}
 	else if (Control == EMediaControl::Pause)
 	{
diff --git a/Engine/Plugins/Media/ElectraPlayer/Source/ElectraProtron/Private/Player/ElectraProtronPlayer.cpp b/Engine/Plugins/Media/ElectraPlayer/Source/ElectraProtron/Private/Player/ElectraProtronPlayer.cpp
--- a/Engine/Plugins/Media/ElectraPlayer/Source/ElectraProtron/Private/Player/ElectraProtronPlayer.cpp
+++ b/Engine/Plugins/Media/ElectraPlayer/Source/ElectraProtron/Private/Player/ElectraProtronPlayer.cpp
@@ -270,7 +270,7 @@ bool FElectraProtronPlayer::CanControl(EMediaControl InControl) const
 {
 	if (InControl == EMediaControl::BlockOnFetch)
 	{
-		return CurrentState == EMediaState::Playing;
+		return CurrentState == EMediaState::Playing || CurrentState == EMediaState::Paused;
 	}
 	else if (InControl == EMediaControl::Pause)
 	{
diff --git a/Engine/Plugins/Media/MediaCompositing/Source/MediaCompositing/Public/MovieSceneMediaTrack.h b/Engine/Plugins/Media/MediaCompositing/Source/MediaCompositing/Public/MovieSceneMediaTrack.h
--- a/Engine/Plugins/Media/MediaCompositing/Source/MediaCompositing/Public/MovieSceneMediaTrack.h
+++ b/Engine/Plugins/Media/MediaCompositing/Source/MediaCompositing/Public/MovieSceneMediaTrack.h
@@ -85,8 +85,6 @@ public:
 	/**
 	 * Force synchronous frame requests *while* manually scrubbing the media track, at the cost of blocking the editor.
 	 * This can be valuable for always maintaining alignment of media & CG while working.
-	 * 
-	 * Note: Only respected by image sequences currently.
 	 */
 	UPROPERTY(EditAnywhere, Category = "Media Track")
 	bool bSynchronousScrubbing = false;
