[Media] Fix MovieSceneMedia Preroll missing the first frame of the video.

For the media section preroll, we can't assign the media texture to the media player because it can be already in use by another section. Consecutive media sections will have the preroll overlapping with the previous section. The media texture can only be presenting one sample queue from one player at a given time, but the prerolling player needs a sample queue to feed it's sample to. The solution is to decouple the sample queue from the media texture and have the media section itself hold the sample queue. When the media texture switches to the media player for the section is coming in, it will also transfer the sample queue, using the one from the section instead of creating a new one internally.

diff --git a/Engine/Plugins/Media/MediaCompositing/Source/MediaCompositing/Private/MovieScene/MovieSceneMediaData.cpp b/Engine/Plugins/Media/MediaCompositing/Source/MediaCompositing/Private/MovieScene/MovieSceneMediaData.cpp
--- a/Engine/Plugins/Media/MediaCompositing/Source/MediaCompositing/Private/MovieScene/MovieSceneMediaData.cpp
+++ b/Engine/Plugins/Media/MediaCompositing/Source/MediaCompositing/Private/MovieScene/MovieSceneMediaData.cpp
@@ -4,13 +4,13 @@
 
 #include "IMediaEventSink.h"
 #include "MediaPlayer.h"
+#include "MediaPlayerFacade.h"
 #include "MediaPlayerProxyInterface.h"
 #include "MediaTexture.h"
 #include "MovieSceneMediaPlayerStore.h"
 #include "UObject/Package.h"
 #include "UObject/UObjectGlobals.h"
 
-
 /* FMediaSectionData structors
  *****************************************************************************/
 
@@ -96,6 +96,9 @@ void FMovieSceneMediaData::Setup(const TSharedPtr<FMovieSceneMediaPlayerStore>&
 	{
 		PlayerProxy.Reset();
 	}
+
+	SampleQueue = MakeShared<FMediaTextureSampleQueue>();
+	MediaPlayer->GetPlayerFacade()->AddVideoSampleSink(SampleQueue.ToSharedRef());
 }
 
 void FMovieSceneMediaData::Initialize(bool bIsEvaluating)
@@ -115,6 +118,13 @@ void FMovieSceneMediaData::TearDown()
 	StopUsingProxyMediaTexture();
 }
 
+TSharedPtr<FMediaTextureSampleQueue> FMovieSceneMediaData::TransferSampleQueue()
+{
+	TSharedPtr<FMediaTextureSampleQueue> OutSampleQueue = SampleQueue;
+	SampleQueue.Reset();
+	return OutSampleQueue;
+}
+
 void FMovieSceneMediaData::StartUsingProxyMediaTexture()
 {
 	if (PlayerProxy != nullptr)
@@ -129,7 +139,7 @@ void FMovieSceneMediaData::StartUsingProxyMediaTexture()
 		}
 		if (ProxyMediaTexture != nullptr)
 		{
-			ProxyMediaTexture->SetMediaPlayer(MediaPlayer);
+			ProxyMediaTexture->SetMediaPlayer(MediaPlayer, TransferSampleQueue());
 		}
 	}
 }
diff --git a/Engine/Plugins/Media/MediaCompositing/Source/MediaCompositing/Private/MovieScene/MovieSceneMediaTemplate.cpp b/Engine/Plugins/Media/MediaCompositing/Source/MediaCompositing/Private/MovieScene/MovieSceneMediaTemplate.cpp
--- a/Engine/Plugins/Media/MediaCompositing/Source/MediaCompositing/Private/MovieScene/MovieSceneMediaTemplate.cpp
+++ b/Engine/Plugins/Media/MediaCompositing/Source/MediaCompositing/Private/MovieScene/MovieSceneMediaTemplate.cpp
@@ -747,7 +747,7 @@ void FMovieSceneMediaSectionTemplate::Initialize(const FMovieSceneEvaluationOper
 				GLog->Logf(ELogVerbosity::Log, TEXT("Setting media player %p on media texture %p"), MediaPlayer, Params.MediaTexture.Get());
 			#endif
 
-			Params.MediaTexture->SetMediaPlayer(MediaPlayer);
+			Params.MediaTexture->SetMediaPlayer(MediaPlayer, SectionData->TransferSampleQueue());
 		}
 		else if (Params.MediaTexture->GetMediaPlayer() == MediaPlayer)
 		{
diff --git a/Engine/Plugins/Media/MediaCompositing/Source/MediaCompositing/Public/MovieSceneMediaData.h b/Engine/Plugins/Media/MediaCompositing/Source/MediaCompositing/Public/MovieSceneMediaData.h
--- a/Engine/Plugins/Media/MediaCompositing/Source/MediaCompositing/Public/MovieSceneMediaData.h
+++ b/Engine/Plugins/Media/MediaCompositing/Source/MediaCompositing/Public/MovieSceneMediaData.h
@@ -2,8 +2,9 @@
 
 #pragma once
 
-#include "Misc/Timespan.h"
 #include "Evaluation/MovieScenePropertyTemplate.h"
+#include "MediaSampleQueue.h"
+#include "Misc/Timespan.h"
 #include "UObject/WeakObjectPtrTemplates.h"
 
 #define UE_API MEDIACOMPOSITING_API
@@ -14,7 +15,6 @@ class UMediaPlayer;
 class UMediaTexture;
 enum class EMediaEvent;
 
-
 /**
  * Persistent data that's stored for each currently evaluating section.
  */
@@ -71,6 +71,9 @@ public:
 	 */
 	UE_API void TearDown();
 
+	/** Retrieve the sample queue and release the ownership. */
+	UE_API TSharedPtr<FMediaTextureSampleQueue> TransferSampleQueue();
+	
 	/**
 	 * Stores if the aspect ratio has been set yet.
 	 */
@@ -116,6 +119,15 @@ private:
 	/** Index of texture allocated from the proxy. */
 	int32 ProxyTextureIndex;
 
+	/**
+	 * Sample queue to be used as video sink for the media player.
+	 * It is owned by the media section because it must be set in the player prior to
+	 * entering the playback section, i.e. in preroll. We can't ue the internal media texture
+	 * sample queue because it is in use by another section during this section's preroll.
+	 * Each section needs it's sample queue independently of the media texture.
+	 */
+	TSharedPtr<FMediaTextureSampleQueue> SampleQueue;
+	
 	/** The time to seek to after the media source is opened. */
 	FTimespan SeekOnOpenTime;
 
diff --git a/Engine/Source/Runtime/MediaAssets/Private/Assets/MediaTexture.cpp b/Engine/Source/Runtime/MediaAssets/Private/Assets/MediaTexture.cpp
--- a/Engine/Source/Runtime/MediaAssets/Private/Assets/MediaTexture.cpp
+++ b/Engine/Source/Runtime/MediaAssets/Private/Assets/MediaTexture.cpp
@@ -135,10 +135,14 @@ int32 UMediaTexture::GetTextureNumMips() const
 
 void UMediaTexture::SetMediaPlayer(UMediaPlayer* NewMediaPlayer)
 {
-	CurrentPlayer = NewMediaPlayer;
-	UpdatePlayerAndQueue();
+	SetMediaPlayer(NewMediaPlayer, /*SampleQueue*/ nullptr);
 }
 
+void UMediaTexture::SetMediaPlayer(UMediaPlayer* NewMediaPlayer, TSharedPtr<FMediaTextureSampleQueue> InSampleQueue)
+{
+	CurrentPlayer = NewMediaPlayer;
+	UpdatePlayerAndQueue(MoveTemp(InSampleQueue));
+}
 
 void UMediaTexture::CacheNextAvailableSampleTime(FTimespan InNextSampleTime)
 {
@@ -539,7 +543,7 @@ void UMediaTexture::UpdateSampleInfo(const TSharedPtr<IMediaTextureSample, ESPMo
 	}
 }
 
-void UMediaTexture::UpdatePlayerAndQueue()
+void UMediaTexture::UpdatePlayerAndQueue(TSharedPtr<FMediaTextureSampleQueue> InSampleQueue)
 {
 	if (UMediaPlayer* CurrentPlayerPtr = CurrentPlayer.Get())
 	{
@@ -553,7 +557,7 @@ void UMediaTexture::UpdatePlayerAndQueue()
 				MediaResource->FlushPendingData();
 			}
 
-			SampleQueue = MakeShared<FMediaTextureSampleQueue, ESPMode::ThreadSafe>();
+			SampleQueue = InSampleQueue ? MoveTemp(InSampleQueue) : MakeShared<FMediaTextureSampleQueue, ESPMode::ThreadSafe>();
 			CurrentPlayerPtr->GetPlayerFacade()->AddVideoSampleSink(SampleQueue.ToSharedRef());
 			CurrentGuid = PlayerGuid;
 		}
diff --git a/Engine/Source/Runtime/MediaAssets/Public/MediaTexture.h b/Engine/Source/Runtime/MediaAssets/Public/MediaTexture.h
--- a/Engine/Source/Runtime/MediaAssets/Public/MediaTexture.h
+++ b/Engine/Source/Runtime/MediaAssets/Public/MediaTexture.h
@@ -148,6 +148,16 @@ public:
 	 */
 	UFUNCTION(BlueprintCallable, Category="Media|MediaTexture")
 	MEDIAASSETS_API void SetMediaPlayer(UMediaPlayer* NewMediaPlayer);
+
+	/**
+	 * Set the media player that provides the video samples and optionally the sample queue
+	 * that receives them. If no sample queue is provided, the media texture creates a new one internally.
+	 *
+	 * @param NewMediaPlayer The player to set.
+	 * @param InSampleQueue Provide an optional sample queue to use instead of creating a new one.
+	 * @see GetMediaPlayer
+	 */
+	MEDIAASSETS_API void SetMediaPlayer(UMediaPlayer* NewMediaPlayer, TSharedPtr<FMediaTextureSampleQueue> InSampleQueue);
 	
 	/**
 	 * Creates a new resource for the texture, and updates any cached references to the resource.
@@ -277,8 +287,11 @@ protected:
 	 */
 	MEDIAASSETS_API void TickResource(FTimespan Timecode);
 
-	/** Update the video sample queue, if necessary. */
-	MEDIAASSETS_API void UpdatePlayerAndQueue();
+	/**
+	 * Update the video sample queue, if necessary.
+	 * @param InSampleQueue Optional sample queue to use.
+	 */
+	MEDIAASSETS_API void UpdatePlayerAndQueue(TSharedPtr<FMediaTextureSampleQueue> InSampleQueue = nullptr);
 
 	/** Update sample info */
 	MEDIAASSETS_API void UpdateSampleInfo(const TSharedPtr<IMediaTextureSample, ESPMode::ThreadSafe>& Sample);
