[MediaCompositing] Media Section: Playback Range edge case fixes

- The playback range can only be applied if the section doesn't have internal looping, i.e. the intersection of playback range and section must be smaller than the full clip. This is what ClampSectionTimeRangeToPlayer implements.
- CalculatePlayerPlaybackRange fix: Intersect the section range with scene range before converting to timespans to avoid issues with start offset.
- SeekPlayer -> allow scrubbing outside of the playback range. Sequencer allows it by default.
- Better consistency between section data's on open seek code and the evaluation template's execution tokens. Make the code in common in MovieSceneMediaPlayerUtils.

diff --git a/Engine/Plugins/Media/MediaCompositing/Source/MediaCompositing/Private/MovieScene/MovieSceneMediaData.cpp b/Engine/Plugins/Media/MediaCompositing/Source/MediaCompositing/Private/MovieScene/MovieSceneMediaData.cpp
--- a/Engine/Plugins/Media/MediaCompositing/Source/MediaCompositing/Private/MovieScene/MovieSceneMediaData.cpp
+++ b/Engine/Plugins/Media/MediaCompositing/Source/MediaCompositing/Private/MovieScene/MovieSceneMediaData.cpp
@@ -8,6 +8,7 @@
 #include "MediaPlayerProxyInterface.h"
 #include "MediaTexture.h"
 #include "MovieSceneMediaPlayerStore.h"
+#include "MovieSceneMediaPlayerUtils.h"
 #include "UObject/Package.h"
 #include "UObject/UObjectGlobals.h"
 
@@ -57,8 +58,12 @@ void FMovieSceneMediaData::SeekOnOpen(FTimespan Time)
 	SeekOnOpenTime = Time;
 }
 
+void FMovieSceneMediaData::Setup(const TSharedPtr<FMovieSceneMediaPlayerStore>& InMediaPlayerStore, const UObject* InSection, UMediaPlayer* OverrideMediaPlayer, UObject* InPlayerProxy, int32 InProxyLayerIndex, int32 InProxyTextureIndex)
+{
+	Setup(InMediaPlayerStore, InSection, OverrideMediaPlayer, InPlayerProxy, InProxyLayerIndex, InProxyTextureIndex, FMovieSceneMediaPlaybackParams());
+}
 
-void FMovieSceneMediaData::Setup(const TSharedPtr<FMovieSceneMediaPlayerStore>& InMediaPlayerStore, const UObject* InMediaSection, UMediaPlayer* OverrideMediaPlayer, UObject* InPlayerProxy, int32 InProxyLayerIndex, int32 InProxyTextureIndex)
+void FMovieSceneMediaData::Setup(const TSharedPtr<FMovieSceneMediaPlayerStore>& InMediaPlayerStore, const UObject* InMediaSection, UMediaPlayer* OverrideMediaPlayer, UObject* InPlayerProxy, int32 InProxyLayerIndex, int32 InProxyTextureIndex, const FMovieSceneMediaPlaybackParams& InPlaybackParams)
 {
 	// Ensure we don't already have a media player set. Setup should only be called once
 	check(!MediaPlayer);
@@ -86,6 +91,7 @@ void FMovieSceneMediaData::Setup(const TSharedPtr<FMovieSceneMediaPlayerStore>&
 	ProxyMediaTexture.Reset();
 	ProxyLayerIndex = InProxyLayerIndex;
 	ProxyTextureIndex = InProxyTextureIndex;
+	PlaybackParams = InPlaybackParams;
 
 	// Do we have a valid proxy object?
 	if ((InPlayerProxy != nullptr) && (InPlayerProxy->Implements<UMediaPlayerProxyInterface>()))
@@ -187,43 +193,33 @@ void FMovieSceneMediaData::HandleMediaPlayerEvent(EMediaEvent Event)
 
 	// Set looping from the media section parameter.
 	// Remark: this must be set to the value it is going to be in the media section, switching it will cause a seek/flush on Electra.
-	MediaPlayer->SetLooping(bIsLooping);
+	MediaPlayer->SetLooping(PlaybackParams.bIsLooping);
 
 	FTimespan MediaTime;
-	
-	if (!PlaybackRange.IsEmpty() && MediaPlayer->SupportsPlaybackTimeRange())
+
+	if (PlaybackParams.bIsLooping)
 	{
-		const TRange<FTimespan> CurrentPlaybackRange = MediaPlayer->GetPlaybackTimeRange(EMediaTimeRangeType::Current);
-		const TRange<FTimespan> FullPlaybackRange = MediaPlayer->GetPlaybackTimeRange(EMediaTimeRangeType::Absolute);
-		// Make sure the playback range from movie scene is valid.
-		const TRange<FTimespan> ClampedRange = TRange<FTimespan>::Intersection(PlaybackRange, FullPlaybackRange);
-		if (CurrentPlaybackRange != ClampedRange)
-		{
-			MediaPlayer->SetPlaybackTimeRange(ClampedRange);
-		}
+		MediaTime = SeekOnOpenTime % Duration;
+	}
+	else
+	{
+		MediaTime = SeekOnOpenTime;
+	}
 
-		// Ensure we are seeking within the specified playback range (very precisely).
-		if (!bIsLooping)
-		{
-			MediaTime = FMath::Clamp(SeekOnOpenTime, ClampedRange.GetLowerBoundValue(), ClampedRange.GetUpperBoundValue());
-		}
-		else
-		{
-			// Loop is defined within the playback range.
-			const FTimespan OffsetTime = FMath::Max(SeekOnOpenTime - ClampedRange.GetLowerBoundValue(), FTimespan::Zero()); 
-			MediaTime = (OffsetTime % ClampedRange.Size<FTimespan>()) + ClampedRange.GetLowerBoundValue();
-		}
+	const FTimespan ClampTolerance = PlaybackParams.FrameDuration * 0.5f;
+
+	if (!PlaybackParams.SectionTimeRange.IsEmpty() && MediaPlayer->SupportsPlaybackTimeRange())
+	{
+		using namespace UE::MovieSceneMediaPlayerUtils;
+		const TRange<FTimespan> AdjustedRange = AdjustPlaybackTimeRange(PlaybackParams.SectionTimeRange, MediaPlayer, PlaybackParams.FrameDuration);
+		SetPlayerPlaybackTimeRange(MediaPlayer, AdjustedRange);
+
+		// Ensure we are seeking within the specified playback range.
+		MediaTime = FMath::Clamp(MediaTime, AdjustedRange.GetLowerBoundValue(), AdjustedRange.GetUpperBoundValue() - ClampTolerance);
 	}
 	else
 	{
-		if (!bIsLooping)
-		{
-			MediaTime = FMath::Clamp(SeekOnOpenTime, FTimespan::Zero(), Duration - FTimespan::FromSeconds(0.001));
-		}
-		else
-		{
-			MediaTime = SeekOnOpenTime % Duration;
-		}
+		MediaTime = FMath::Clamp(MediaTime, FTimespan::Zero(), Duration - ClampTolerance);
 	}
 
 	MediaPlayer->SetRate(0.0f);
diff --git a/Engine/Plugins/Media/MediaCompositing/Source/MediaCompositing/Private/MovieScene/MovieSceneMediaPlayerUtils.cpp b/Engine/Plugins/Media/MediaCompositing/Source/MediaCompositing/Private/MovieScene/MovieSceneMediaPlayerUtils.cpp
new file mode 100644
--- /dev/null
+++ b/Engine/Plugins/Media/MediaCompositing/Source/MediaCompositing/Private/MovieScene/MovieSceneMediaPlayerUtils.cpp
@@ -0,0 +1,91 @@
+// Copyright Epic Games, Inc. All Rights Reserved.
+
+#include "MovieSceneMediaPlayerUtils.h"
+
+#include "MediaPlayer.h"
+#include "MovieSceneMediaData.h"
+
+namespace UE::MovieSceneMediaPlayerUtils
+{
+	TRange<FTimespan> AdjustPlaybackTimeRange(const TRange<FTimespan>& InSectionPlaybackRange, UMediaPlayer* InMediaPlayer, const FTimespan& InFrameDuration)
+	{
+		if (!InMediaPlayer->SupportsPlaybackTimeRange() || InSectionPlaybackRange.IsEmpty())
+		{
+			return InSectionPlaybackRange;
+		}
+
+		const TRange<FTimespan> FullPlaybackRange = InMediaPlayer->GetPlaybackTimeRange(EMediaTimeRangeType::Absolute);
+
+		// Expand half a sequencer frame on each side to give some tolerance.
+		const FTimespan RangeTolerance = InFrameDuration * 0.5f;
+		TRange<FTimespan> FullPlaybackRangeWithTolerance(
+			TRangeBound<FTimespan>::Inclusive(FullPlaybackRange.GetLowerBoundValue() - RangeTolerance),
+			TRangeBound<FTimespan>::Inclusive(FullPlaybackRange.GetUpperBoundValue() + RangeTolerance));
+
+		// Only apply clamping if the section range is inside (with tolerance) the player's full range.
+		if (FullPlaybackRangeWithTolerance.Contains(InSectionPlaybackRange))
+		{
+			return TRange<FTimespan>::Intersection(InSectionPlaybackRange, FullPlaybackRange);
+		}
+
+		// Media player only supports the specification of a playback range that is within the clip duration.
+		return FullPlaybackRange;
+	}
+
+	void SetPlayerPlaybackTimeRange(UMediaPlayer* InMediaPlayer, const TRange<FTimespan>& InPlaybackRange)
+	{
+		if (!InMediaPlayer->SupportsPlaybackTimeRange())
+		{
+			return;
+		}
+		
+		const TRange<FTimespan> CurrentPlaybackRange = InMediaPlayer->GetPlaybackTimeRange(EMediaTimeRangeType::Current);
+
+		if (!InPlaybackRange.IsEmpty())
+		{
+			if (CurrentPlaybackRange != InPlaybackRange)
+			{
+				InMediaPlayer->SetPlaybackTimeRange(InPlaybackRange);
+			}
+		}
+		else
+		{
+			const TRange<FTimespan> FullPlaybackRange = InMediaPlayer->GetPlaybackTimeRange(EMediaTimeRangeType::Absolute);
+
+			// Restore playback range to full clip length (default)
+			if (CurrentPlaybackRange != FullPlaybackRange)
+			{
+				InMediaPlayer->SetPlaybackTimeRange(FullPlaybackRange);
+			}
+		}
+	}
+
+	FTimespan ClampTimeToPlaybackRange(const FTimespan& InMediaTime, UMediaPlayer* InMediaPlayer, const FMovieSceneMediaPlaybackParams& InPlaybackParams)
+	{
+		const FTimespan Duration = InMediaPlayer->GetDuration();
+		if (Duration == FTimespan::Zero())
+		{
+			return InMediaTime;
+		}
+
+		FTimespan MediaTime = InMediaTime;
+
+		if (InPlaybackParams.bIsLooping)
+		{
+			MediaTime = MediaTime % Duration;
+		}
+
+		const FTimespan ClampTolerance = InPlaybackParams.FrameDuration * 0.5f;
+		
+		if (!InPlaybackParams.SectionTimeRange.IsEmpty() && InMediaPlayer->SupportsPlaybackTimeRange())
+		{
+			const TRange<FTimespan> AdjustedRange = AdjustPlaybackTimeRange(InPlaybackParams.SectionTimeRange, InMediaPlayer, InPlaybackParams.FrameDuration);
+			MediaTime = FMath::Clamp(MediaTime, AdjustedRange.GetLowerBoundValue(), AdjustedRange.GetUpperBoundValue() - ClampTolerance);
+		}
+		else
+		{
+			MediaTime = FMath::Clamp(MediaTime, FTimespan::Zero(), Duration - ClampTolerance);
+		}
+		return MediaTime;
+	}
+}
diff --git a/Engine/Plugins/Media/MediaCompositing/Source/MediaCompositing/Private/MovieScene/MovieSceneMediaPlayerUtils.h b/Engine/Plugins/Media/MediaCompositing/Source/MediaCompositing/Private/MovieScene/MovieSceneMediaPlayerUtils.h
new file mode 100644
--- /dev/null
+++ b/Engine/Plugins/Media/MediaCompositing/Source/MediaCompositing/Private/MovieScene/MovieSceneMediaPlayerUtils.h
@@ -0,0 +1,39 @@
+// Copyright Epic Games, Inc. All Rights Reserved.
+
+#pragma once
+
+#include "Misc/Timespan.h"
+
+class UMediaPlayer;
+struct FMovieSceneMediaPlaybackParams;
+
+namespace UE::MovieSceneMediaPlayerUtils
+{
+	/**
+	 * Utility function to adjust the playback range of the section according to proper looping logic
+	 * and what is possible to do with the current player's playback range.
+	 *
+	 * In more detail, the current player's playback range must be smaller or equal to the full clip length.
+	 * This limits us to the case where the section's playback range is also smaller than the full clip. We
+	 * can't specify a playback range that would be 2.5 times the full clip length for instance. In that case,
+	 * another method of looping will have to be implemented (ex: using a second player).
+	 *
+	 * @param InSectionPlaybackRange Playback range derived from the media section
+	 * @param InMediaPlayer Player to adjust against
+	 * @param InFrameDuration Sequencer frame duration, used for tolerance
+	 * @return 
+	 */
+	TRange<FTimespan> AdjustPlaybackTimeRange(const TRange<FTimespan>& InSectionPlaybackRange, UMediaPlayer* InMediaPlayer, const FTimespan& InFrameDuration);
+
+	/**
+	 * Utility function to set the specified playback time range.
+	 * @param InMediaPlayer Player to use
+	 * @param InPlaybackRange Specifies the playback range to set. Must be already clamped. If empty, player's full range is restored.
+	 */
+	void SetPlayerPlaybackTimeRange(UMediaPlayer* InMediaPlayer, const TRange<FTimespan>& InPlaybackRange);
+
+	/**
+	 * Ensures that the given time is properly clamped to player's playback range.
+	 */
+	FTimespan ClampTimeToPlaybackRange(const FTimespan& InMediaTime, UMediaPlayer* InMediaPlayer, const FMovieSceneMediaPlaybackParams& InPlaybackParams);
+}
\ No newline at end of file
diff --git a/Engine/Plugins/Media/MediaCompositing/Source/MediaCompositing/Private/MovieScene/MovieSceneMediaTemplate.cpp b/Engine/Plugins/Media/MediaCompositing/Source/MediaCompositing/Private/MovieScene/MovieSceneMediaTemplate.cpp
--- a/Engine/Plugins/Media/MediaCompositing/Source/MediaCompositing/Private/MovieScene/MovieSceneMediaTemplate.cpp
+++ b/Engine/Plugins/Media/MediaCompositing/Source/MediaCompositing/Private/MovieScene/MovieSceneMediaTemplate.cpp
@@ -18,6 +18,7 @@
 
 #include "MovieSceneMediaData.h"
 #include "MovieSceneMediaPlayerStore.h"
+#include "MovieSceneMediaPlayerUtils.h"
 #include "MovieSceneMediaSection.h"
 #include "MovieSceneMediaTrack.h"
 
@@ -113,33 +114,39 @@ namespace UE::MediaSectionBaseMediaUtils
 		return IsVideoSampleAvailableForTime(MediaPlayerFacade, InTime);
 	}
 
+	/**
+	 * Utility to convert a frame number to a timespan at the given frame rate. 
+	 */
+	FTimespan FrameNumberToTimespan(const FFrameNumber& InFrameNumber, const FFrameRate& InFrameRate)
+	{
+		// Caution: A larger denominator limits the frame number range. 
+		// Example: for 59.94, 29.97 or 23.975, because denominator is 1001, the maximum frame number is going to be 1073741823 (which is ~207 days at 59.94).
+		const int64 DenominatorTicks = InFrameRate.Denominator * ETimespan::TicksPerSecond;
+		// Using Ceil to ensure the tick value is part of the desired frame, and not the frame before
+		// regardless of the internal media player's implementation of "time to frame".
+		return FTimespan(FMath::DivideAndRoundUp(int64(InFrameNumber.Value * DenominatorTicks), int64(InFrameRate.Numerator)));
+	}
+	
 	/**
 	 * Utility to convert a sequencer's (absolute) frame number to a player's timespan.
 	 */
 	FTimespan SequencerFrameToPlayerTime(const FMovieSceneMediaSectionParams& InMediaSectionParams, const FFrameRate& InSequencerFrameRate, const FFrameNumber& InSequencerFrame)
 	{
 		const FFrameNumber PlayerFrame = InSequencerFrame - InMediaSectionParams.SectionStartFrame + InMediaSectionParams.StartFrameOffset;
-		return FTimespan::FromSeconds(InSequencerFrameRate.AsSeconds(PlayerFrame));
-	}
-
-	/** Converts specified time to a range bound of given type (inclusive or exclusive). */
-	template <typename T>
-	TRangeBound<T> ToRangeBound(bool bInIsExclusive, const T& InTime)
-	{
-		return bInIsExclusive ? TRangeBound<T>::Exclusive(InTime) : TRangeBound<T>::Inclusive(InTime); 
+		return FrameNumberToTimespan(PlayerFrame, InSequencerFrameRate);
 	}
 
 	/**
 	 * Converts the given sequencer frame number range in media player's time span range.
 	 * @param InMediaSectionParams Media Section information
 	 * @param InSequencerFrameRate Sequencer frame rate to convert frame number in time span
-	 * @param InFrameRange Frame Number range to convert
+	 * @param InFrameRange Frame Number range to convert, must be all inclusive.
 	 */
 	TRange<FTimespan> SequencerFrameRangeToPlayerRange(const FMovieSceneMediaSectionParams& InMediaSectionParams, const FFrameRate& InSequencerFrameRate, const TRange<FFrameNumber>& InFrameRange)
 	{
 		const FTimespan LowerTime = SequencerFrameToPlayerTime(InMediaSectionParams, InSequencerFrameRate, InFrameRange.GetLowerBoundValue());
 		const FTimespan UpperTime = SequencerFrameToPlayerTime(InMediaSectionParams, InSequencerFrameRate, InFrameRange.GetUpperBoundValue());
-		return TRange<FTimespan>(ToRangeBound(InFrameRange.GetLowerBound().IsExclusive(), LowerTime), ToRangeBound(InFrameRange.GetUpperBound().IsExclusive(), UpperTime));
+		return TRange<FTimespan>(TRangeBound<FTimespan>::Inclusive(LowerTime), TRangeBound<FTimespan>::Inclusive(UpperTime));
 	}
 
 	/**
@@ -162,80 +169,62 @@ namespace UE::MediaSectionBaseMediaUtils
 	}
 	
 	/**
-	 * Calculates the intersection of the section's bounds and scene's playback range.
+	 * Calculates the intersection of the section's bounds and scene's playback range and convert to player's time range.
 	 */
-	TRange<FTimespan> CalculatePlayerPlaybackRange(const FMovieSceneMediaSectionParams& InMediaSectionParams, const UMovieSceneMediaSection* InMediaSection, const FFrameRate& InSequencerFrameRate)
+	TRange<FTimespan> CalculateSectionPlaybackTimeRange(const FMovieSceneMediaSectionParams& InMediaSectionParams, const UMovieSceneMediaSection* InMediaSection, const FFrameRate& InSequencerFrameRate)
 	{
-		// Todo: validate this. See GetTrueRange(). For consistency, should be: { lower <= time < upper }
-		const TRange<FFrameNumber> SectionFrameRange = TRange<FFrameNumber>(TRangeBound<FFrameNumber>::Inclusive(InMediaSectionParams.SectionStartFrame), TRangeBound<FFrameNumber>::Exclusive(InMediaSectionParams.SectionEndFrame));
-		const TRange<FTimespan> SectionTimeRange = SequencerFrameRangeToPlayerRange(InMediaSectionParams, InSequencerFrameRate, ToAllInclusiveRange(SectionFrameRange));
-
-		const UMovieSceneTrack* Track = InMediaSection->GetTypedOuter<UMovieSceneTrack>();
-		const UMovieScene* Scene = Track ? Track->GetTypedOuter<UMovieScene>() : nullptr;
-
-		if (Scene)
+		TRange<FFrameNumber> SectionFrameRange = InMediaSection->GetTrueRange();
+		
+		if (const UMovieSceneTrack* Track = InMediaSection->GetTypedOuter<UMovieSceneTrack>())
 		{
-			const TRange<FTimespan> SceneTimeRange = SequencerFrameRangeToPlayerRange(InMediaSectionParams, InSequencerFrameRate, ToAllInclusiveRange(Scene->GetPlaybackRange()));
-			return TRange<FTimespan>::Intersection(SectionTimeRange, SceneTimeRange);
+			if (const UMovieScene* Scene = Track->GetTypedOuter<UMovieScene>())
+			{
+				SectionFrameRange = TRange<FFrameNumber>::Intersection(Scene->GetPlaybackRange(), SectionFrameRange);
+			}
 		}
 
-		return SectionTimeRange;
+		return SequencerFrameRangeToPlayerRange(InMediaSectionParams, InSequencerFrameRate, ToAllInclusiveRange(SectionFrameRange));
 	}
 
 	/** Utility function to set the specified playback time range */
-	void SetPlayerPlaybackTimeRange(UMediaPlayer* InMediaPlayer, const TRange<FTimespan>& InPlaybackRange)
+	void SetPlayerPlaybackTimeRange(UMediaPlayer* InMediaPlayer, const FMovieSceneMediaPlaybackParams& InPlaybackParams)
 	{
-		if (InMediaPlayer->SupportsPlaybackTimeRange())
-		{
-			const TRange<FTimespan> CurrentPlaybackRange = InMediaPlayer->GetPlaybackTimeRange(EMediaTimeRangeType::Current);
-			const TRange<FTimespan> FullPlaybackRange = InMediaPlayer->GetPlaybackTimeRange(EMediaTimeRangeType::Absolute);
-
-			if (!InPlaybackRange.IsEmpty())
-			{
-				// Make sure the playback range from movie scene is valid.
-				const TRange<FTimespan> ClampedRange = TRange<FTimespan>::Intersection(InPlaybackRange, FullPlaybackRange);
-				
-				if (CurrentPlaybackRange != ClampedRange)
-				{
-					InMediaPlayer->SetPlaybackTimeRange(ClampedRange);
-				}
-			}
-			else
-			{
-				// Restore playback range to full clip length (default)
-				if (CurrentPlaybackRange != FullPlaybackRange)
-				{
-					InMediaPlayer->SetPlaybackTimeRange(FullPlaybackRange);
-				}
-			}
-		}
+		using namespace UE::MovieSceneMediaPlayerUtils;
+		const TRange<FTimespan> AdjustedRange = AdjustPlaybackTimeRange(InPlaybackParams.SectionTimeRange, InMediaPlayer, InPlaybackParams.FrameDuration);
+		UE::MovieSceneMediaPlayerUtils::SetPlayerPlaybackTimeRange(InMediaPlayer, AdjustedRange);
 	}
 
 	/**
 	 * Helper function to set BlockOnTimeRange.
-	 * If a PlaybackRange is specified, the block on range must be within the playback range,
-	 * very precisely to avoid a blocking time out.
-	 * TODO: this is probably best done in FMediaPlayerFacade. 
+	 * @param InMediaPlayer	Media player to use
+	 * @param InCurrentTime Frame time that will be the start of the range to block on. Important: Must be clamped to playback range.
+	 * @param InPlaybackParams Section derived playback parameters
 	 */
-	void SetPlayerBlockOnTimeRange(UMediaPlayer* InMediaPlayer, const FTimespan& InCurrentTime, const FTimespan& InFrameDuration, const TRange<FTimespan>& InPlaybackRange)
+	void SetPlayerBlockOnTimeRange(UMediaPlayer* InMediaPlayer, const FTimespan& InCurrentTime, const FMovieSceneMediaPlaybackParams& InPlaybackParams)
 	{
-		FTimespan CurrentTime = InCurrentTime;
-		if (!InPlaybackRange.IsEmpty() && InMediaPlayer->SupportsPlaybackTimeRange())
-		{
-			// Try to ensure the specified time range is within the playback range very precisely. (otherwise it will block forever)
-			const TRange<FTimespan> CurrentPlaybackRange = InMediaPlayer->GetPlaybackTimeRange(EMediaTimeRangeType::Current);
-			if (!InMediaPlayer->IsLooping())
-			{
-				CurrentTime = FMath::Clamp(CurrentTime, CurrentPlaybackRange.GetLowerBoundValue(), CurrentPlaybackRange.GetUpperBoundValue() - InFrameDuration * 0.5);
-			}
-			else
-			{
-				// TODO
-				// ... Looping case is complex, need to match FMediaPlayerFacade::FBlockOnRange::GetRange() ...
-			}
-		}
+		FTimespan CurrentTime = UE::MovieSceneMediaPlayerUtils::ClampTimeToPlaybackRange(InCurrentTime, InMediaPlayer, InPlaybackParams);
+		InMediaPlayer->SetBlockOnTimeRange(TRange<FTimespan>(CurrentTime, CurrentTime + InPlaybackParams.FrameDuration));
+	}
+
+	/** Calculate the timespan of a frame duration given the sequencer's framerate */
+	FTimespan GetFrameDuration(const FMovieSceneContext& InContext)
+	{
+		const FFrameRate FrameRate = InContext.GetFrameRate();
+		// With zero-length frames (which can occur occasionally), we use the fixed frame time, matching previous behavior.
+		const double FrameDurationInSeconds = FMath::Max(FrameRate.AsSeconds(FFrameTime(1)), (InContext.GetRange().Size<FFrameTime>()) / FrameRate);
+		const int64 FrameDurationTicks = static_cast<int64>(FrameDurationInSeconds * ETimespan::TicksPerSecond);
+		return FTimespan(FrameDurationTicks);
+	}
 
-		InMediaPlayer->SetBlockOnTimeRange(TRange<FTimespan>(CurrentTime, CurrentTime + InFrameDuration));
+	/** Prepare the playback parameters for the execution token. */
+	FMovieSceneMediaPlaybackParams MakePlaybackParams(const FMovieSceneContext& InContext, const FMovieSceneMediaSectionParams& InParams, const UMovieSceneMediaSection* InMediaSection)
+	{
+		FMovieSceneMediaPlaybackParams Params;
+		Params.SectionTimeRange = CVarUpdatePlaybackRange.GetValueOnGameThread() ?
+			CalculateSectionPlaybackTimeRange(InParams, InMediaSection, InContext.GetFrameRate()) : TRange<FTimespan>::Empty(); 
+		Params.FrameDuration = GetFrameDuration(InContext);
+		Params.bIsLooping = InParams.bLooping;
+		return Params;
 	}
 }
 
@@ -243,11 +232,12 @@ namespace UE::MediaSectionBaseMediaUtils
 struct FMediaSectionBaseExecutionToken
 	: IMovieSceneExecutionToken
 {
-	FMediaSectionBaseExecutionToken(UMediaSource* InMediaSource, const FMediaSectionBaseMediaSourceOptions& InMediaSourceOptions, const FMovieSceneObjectBindingID& InMediaSourceProxy, int32 InMediaSourceProxyIndex)
+	FMediaSectionBaseExecutionToken(UMediaSource* InMediaSource, const FMediaSectionBaseMediaSourceOptions& InMediaSourceOptions, const FMovieSceneMediaPlaybackParams& InPlaybackParams, const FMovieSceneObjectBindingID& InMediaSourceProxy, int32 InMediaSourceProxyIndex)
 		: BaseMediaSource(InMediaSource)
 		, MediaSourceProxy(InMediaSourceProxy)
 		, MediaSourceProxyIndex(InMediaSourceProxyIndex)
 		, BaseMediaSourceOptions(InMediaSourceOptions)
+		, PlaybackParams(InPlaybackParams)
 	{
 	}
 
@@ -275,18 +265,22 @@ struct FMediaSectionBaseExecutionToken
 	 */
 	const FMediaSectionBaseMediaSourceOptions& GetBaseMediaSourceOptions() const { return BaseMediaSourceOptions; }
 
+	/** Get the section's playback parameters. */
+	const FMovieSceneMediaPlaybackParams& GetPlaybackParams() const { return PlaybackParams; }
+
 private:
 	UMediaSource* BaseMediaSource;
 	FMovieSceneObjectBindingID MediaSourceProxy;
 	int32 MediaSourceProxyIndex = 0;
 	FMediaSectionBaseMediaSourceOptions BaseMediaSourceOptions;
+	FMovieSceneMediaPlaybackParams PlaybackParams;
 };
 
 struct FMediaSectionPreRollExecutionToken
 	: FMediaSectionBaseExecutionToken
 {
-	FMediaSectionPreRollExecutionToken(UMediaSource* InMediaSource, const FMediaSectionBaseMediaSourceOptions& InMediaSourceOptions, FMovieSceneObjectBindingID InMediaSourceProxy, int32 InMediaSourceProxyIndex, FTimespan InStartTimeSeconds)
-		: FMediaSectionBaseExecutionToken(InMediaSource, InMediaSourceOptions, InMediaSourceProxy, InMediaSourceProxyIndex)
+	FMediaSectionPreRollExecutionToken(UMediaSource* InMediaSource, const FMediaSectionBaseMediaSourceOptions& InMediaSourceOptions, const FMovieSceneMediaPlaybackParams& InPlaybackParams, FMovieSceneObjectBindingID InMediaSourceProxy, int32 InMediaSourceProxyIndex, FTimespan InStartTimeSeconds)
+		: FMediaSectionBaseExecutionToken(InMediaSource, InMediaSourceOptions, InPlaybackParams, InMediaSourceProxy, InMediaSourceProxyIndex)
 		, StartTime(InStartTimeSeconds)
 	{ }
 
@@ -327,7 +321,10 @@ struct FMediaSectionPreRollExecutionToken
 		const bool bMoveToNewTime = Context.GetStatus() != EMovieScenePlayerStatus::Playing || (Context.GetStatus() == EMovieScenePlayerStatus::Playing && Context.HasJumped());
 		if (bMoveToNewTime)
 		{
-			MediaPlayer->SetRate(0.0f);
+			if (MediaPlayer->GetRate() != 0.0f)	// avoids Electra spamming when scrubbing in the preroll.
+			{
+				MediaPlayer->SetRate(0.0f);
+			}
 			MediaPlayer->Seek(StartTime);
 			MediaPlayer->SetBlockOnTimeRange(TRange<FTimespan>::Empty());
 		}
@@ -342,25 +339,30 @@ private:
 struct FMediaSectionExecutionToken
 	: FMediaSectionBaseExecutionToken
 {
-	FMediaSectionExecutionToken(UMediaSource* InMediaSource, const FMediaSectionBaseMediaSourceOptions& InMediaSourceOptions, FMovieSceneObjectBindingID InMediaSourceProxy, int32 InMediaSourceProxyIndex, float InProxyTextureBlend, bool bInCanPlayerBeOpen, FTimespan InCurrentTime, FTimespan InFrameDuration, const TRange<FTimespan>& InPlaybackRange)
-		: FMediaSectionBaseExecutionToken(InMediaSource, InMediaSourceOptions, InMediaSourceProxy, InMediaSourceProxyIndex)
+	FMediaSectionExecutionToken(UMediaSource* InMediaSource, const FMediaSectionBaseMediaSourceOptions& InMediaSourceOptions, const FMovieSceneMediaPlaybackParams& InPlaybackParams, FMovieSceneObjectBindingID InMediaSourceProxy, int32 InMediaSourceProxyIndex, float InProxyTextureBlend, bool bInCanPlayerBeOpen, FTimespan InCurrentTime)
+		: FMediaSectionBaseExecutionToken(InMediaSource, InMediaSourceOptions, InPlaybackParams, InMediaSourceProxy, InMediaSourceProxyIndex)
 		, CurrentTime(InCurrentTime)
-		, FrameDuration(InFrameDuration)
-		, PlaybackRange(InPlaybackRange)
-		, PlaybackRate(1.0f)
 		, ProxyTextureBlend(InProxyTextureBlend)
 		, bCanPlayerBeOpen(bInCanPlayerBeOpen)
 	{ }
 
-	static void SeekPlayer(const FMovieSceneContext& InContext, UMediaPlayer* InMediaPlayer, const FTimespan& InMediaTime)
+	void SeekPlayer(const FMovieSceneContext& InContext, UMediaPlayer* InMediaPlayer, const FTimespan& InMediaTime) const
 	{
 		if (InContext.GetStatus() == EMovieScenePlayerStatus::Scrubbing || InContext.GetStatus() == EMovieScenePlayerStatus::Stopped)
 		{
+			// Scrubbing outside the playback range is allowed, in this case, we need to reset it.
+			if (!GetPlaybackParams().SectionTimeRange.IsEmpty() && !GetPlaybackParams().SectionTimeRange.Contains(InMediaTime))
+			{
+				UE::MovieSceneMediaPlayerUtils::SetPlayerPlaybackTimeRange(InMediaPlayer, TRange<FTimespan>::Empty());
+			}
+
 			InMediaPlayer->Scrub(InMediaTime);
 		}
 		else
 		{
-			InMediaPlayer->Seek(InMediaTime);
+			using namespace UE::MovieSceneMediaPlayerUtils;
+			const FTimespan MediaTime = ClampTimeToPlaybackRange(InMediaTime, InMediaPlayer, GetPlaybackParams());
+			InMediaPlayer->Seek(MediaTime);
 		}
 	}
 
@@ -429,7 +431,7 @@ struct FMediaSectionExecutionToken
 			}
 
 			// Setup an initial blocking range - MediaFramework will block (even through the opening process) in its next tick...
-			SetPlayerBlockOnTimeRange(MediaPlayer, CurrentTime, FrameDuration, PlaybackRange);
+			SetPlayerBlockOnTimeRange(MediaPlayer, CurrentTime, GetPlaybackParams());
 			MediaPlayer->OpenSourceWithOptions(MediaSource, Options);
 			return;
 		}
@@ -439,7 +441,7 @@ struct FMediaSectionExecutionToken
 		if (MediaPlayer->IsPreparing())
 		{
 			SectionData.SeekOnOpen(CurrentTime);
-			SetPlayerBlockOnTimeRange(MediaPlayer, CurrentTime, FrameDuration, PlaybackRange);
+			SetPlayerBlockOnTimeRange(MediaPlayer, CurrentTime, GetPlaybackParams());
 			return;
 		}
 
@@ -454,7 +456,7 @@ struct FMediaSectionExecutionToken
 		// update media player
 		//
 
-		SetPlayerPlaybackTimeRange(MediaPlayer, PlaybackRange);
+		SetPlayerPlaybackTimeRange(MediaPlayer, GetPlaybackParams());
 
 		// Setup media time (used for seeks)
 		FTimespan MediaTime = CurrentTime;
@@ -517,7 +519,7 @@ struct FMediaSectionExecutionToken
 		}
 		else
 		{
-			if (MediaPlayer->IsPlaying())
+			if (MediaPlayer->IsPlaying() && MediaPlayer->GetRate() != 0.0f)
 			{
 				MediaPlayer->SetRate(0.0f);
 			}
@@ -535,16 +537,13 @@ struct FMediaSectionExecutionToken
 		{
 			// Set blocking range / time-range to display
 			// (we always use the full current time for this, any adjustments to player timestamps are done internally)
-			SetPlayerBlockOnTimeRange(MediaPlayer, CurrentTime, FrameDuration, PlaybackRange);
+			SetPlayerBlockOnTimeRange(MediaPlayer, CurrentTime, GetPlaybackParams());
 		}
 	}
 
 private:
 
 	FTimespan CurrentTime;
-	FTimespan FrameDuration;
-	TRange<FTimespan> PlaybackRange;
-	float PlaybackRate;
 	float ProxyTextureBlend;
 	bool bCanPlayerBeOpen;
 };
@@ -599,33 +598,22 @@ void FMovieSceneMediaSectionTemplate::Evaluate(const FMovieSceneEvaluationOperan
 	{
 		const FFrameRate FrameRate = Context.GetFrameRate();
 		const FFrameNumber StartFrame = Context.HasPreRollEndTime() ? Context.GetPreRollEndFrame() - Params.SectionStartFrame + Params.StartFrameOffset : Params.StartFrameOffset;
-		// Using Ceil to ensure the tick value is part of the desired frame, and not the frame before.
-		const int64 StartTicks = FMath::CeilToInt64(FrameRate.AsSeconds(StartFrame) * ETimespan::TicksPerSecond);
+		const FTimespan StartTime = UE::MediaSectionBaseMediaUtils::FrameNumberToTimespan(StartFrame, FrameRate);
 
 		FMediaSectionBaseMediaSourceOptions Options;
 		Options.CacheSettings = Params.CacheSettings;
 		Options.bSynchronousScrubbing = false; // not relevant during pre-roll
-		ExecutionTokens.Add(FMediaSectionPreRollExecutionToken(MediaSource, Options, Params.MediaSourceProxy, Params.MediaSourceProxyIndex, FTimespan(StartTicks)));
+		const FMovieSceneMediaPlaybackParams PlaybackParams = UE::MediaSectionBaseMediaUtils::MakePlaybackParams(Context, Params, MediaSection);
+		ExecutionTokens.Add(FMediaSectionPreRollExecutionToken(MediaSource, Options, PlaybackParams, Params.MediaSourceProxy, Params.MediaSourceProxyIndex, StartTime));
 	}
 	else if (!Context.IsPostRoll() && (Context.GetTime().FrameNumber < Params.SectionEndFrame))
 	{
 		const FFrameRate FrameRate = Context.GetFrameRate();
-		const FFrameTime FrameTime(Context.GetTime().FrameNumber - Params.SectionStartFrame + Params.StartFrameOffset);
-		// Using Ceil to ensure the tick value is part of the desired frame, and not the frame before.
-		const int64 FrameTicks = FMath::CeilToInt64(FrameRate.AsSeconds(FrameTime) * ETimespan::TicksPerSecond);
-
-		// With zero-length frames (which can occur occasionally), we use the fixed frame time, matching previous behavior.
-		const double FrameDurationInSeconds = FMath::Max(FrameRate.AsSeconds(FFrameTime(1)), (Context.GetRange().Size<FFrameTime>()) / Context.GetFrameRate());
-		const int64 FrameDurationTicks = static_cast<int64>(FrameDurationInSeconds * ETimespan::TicksPerSecond);
-
+		const FFrameNumber FrameNumber = Context.GetTime().FrameNumber - Params.SectionStartFrame + Params.StartFrameOffset;
+		const FTimespan FrameTime = UE::MediaSectionBaseMediaUtils::FrameNumberToTimespan(FrameNumber, FrameRate);
+		
 		const float ProxyTextureBlend = MediaSection->EvaluateEasing(Context.GetTime());
 
-		TRange<FTimespan> PlayerPlaybackRange;
-		if (CVarUpdatePlaybackRange.GetValueOnGameThread())
-		{
-			PlayerPlaybackRange = UE::MediaSectionBaseMediaUtils::CalculatePlayerPlaybackRange(Params, MediaSection, FrameRate);
-		}
-
 		#if MOVIESCENEMEDIATEMPLATE_TRACE_EVALUATION
 			GLog->Logf(ELogVerbosity::Log, TEXT("Evaluating frame %i+%f, FrameRate %i/%i, FrameTicks %d, FrameDurationTicks %d"),
 				Context.GetTime().GetFrame().Value,
@@ -642,7 +630,8 @@ void FMovieSceneMediaSectionTemplate::Evaluate(const FMovieSceneEvaluationOperan
 		FMediaSectionBaseMediaSourceOptions Options;
 		Options.CacheSettings = Params.CacheSettings;
 		Options.bSynchronousScrubbing = MediaTrack ? MediaTrack->bSynchronousScrubbing : false;
-		ExecutionTokens.Add(FMediaSectionExecutionToken(MediaSource, Options, Params.MediaSourceProxy, Params.MediaSourceProxyIndex, ProxyTextureBlend, bCanPlayerBeOpen, FTimespan(FrameTicks), FTimespan(FrameDurationTicks), PlayerPlaybackRange));
+		const FMovieSceneMediaPlaybackParams PlaybackParams = UE::MediaSectionBaseMediaUtils::MakePlaybackParams(Context, Params, MediaSection);
+		ExecutionTokens.Add(FMediaSectionExecutionToken(MediaSource, Options, PlaybackParams, Params.MediaSourceProxy, Params.MediaSourceProxyIndex, ProxyTextureBlend, bCanPlayerBeOpen, FrameTime));
 	}
 }
 
@@ -694,14 +683,10 @@ void FMovieSceneMediaSectionTemplate::Initialize(const FMovieSceneEvaluationOper
 		// Add section data.
 		SectionData = &PersistentData.AddSectionData<FMovieSceneMediaData>();
 		const TSharedPtr<FMovieSceneMediaPlayerStore> MediaPlayerStore = FMovieSceneMediaPlayerStoreContainer::GetOrAdd(PersistentData).GetMediaPlayerStore();
-		SectionData->Setup(MediaPlayerStore, MediaSection, MediaPlayer, PlayerProxy, ProxyLayerIndex, ProxyTextureIndex);
-
 		// Prepare these values in the section data to have them set in the player when it is opened (as early as possible before the first seek).
-		SectionData->bIsLooping = Params.bLooping;
-		if (CVarUpdatePlaybackRange.GetValueOnGameThread())
-		{
-			SectionData->PlaybackRange = UE::MediaSectionBaseMediaUtils::CalculatePlayerPlaybackRange(Params, MediaSection, Context.GetFrameRate());
-		}
+		FMovieSceneMediaPlaybackParams PlaybackParams = UE::MediaSectionBaseMediaUtils::MakePlaybackParams(Context, Params, MediaSection);
+
+		SectionData->Setup(MediaPlayerStore, MediaSection, MediaPlayer, PlayerProxy, ProxyLayerIndex, ProxyTextureIndex, PlaybackParams);
 	}
 
 	if (!ensure(SectionData != nullptr))
diff --git a/Engine/Plugins/Media/MediaCompositing/Source/MediaCompositing/Private/MovieScene/MovieSceneMediaTemplate.h b/Engine/Plugins/Media/MediaCompositing/Source/MediaCompositing/Private/MovieScene/MovieSceneMediaTemplate.h
--- a/Engine/Plugins/Media/MediaCompositing/Source/MediaCompositing/Private/MovieScene/MovieSceneMediaTemplate.h
+++ b/Engine/Plugins/Media/MediaCompositing/Source/MediaCompositing/Private/MovieScene/MovieSceneMediaTemplate.h
@@ -4,7 +4,6 @@
 
 #include "Channels/MovieSceneFloatChannel.h"
 #include "Evaluation/MovieSceneEvalTemplate.h"
-#include "Evaluation/MovieScenePropertyTemplate.h"
 #include "MovieSceneObjectBindingID.h"
 #include "MediaSourceOptions.h"
 
diff --git a/Engine/Plugins/Media/MediaCompositing/Source/MediaCompositing/Public/MovieSceneMediaData.h b/Engine/Plugins/Media/MediaCompositing/Source/MediaCompositing/Public/MovieSceneMediaData.h
--- a/Engine/Plugins/Media/MediaCompositing/Source/MediaCompositing/Public/MovieSceneMediaData.h
+++ b/Engine/Plugins/Media/MediaCompositing/Source/MediaCompositing/Public/MovieSceneMediaData.h
@@ -15,6 +15,27 @@ class UMediaPlayer;
 class UMediaTexture;
 enum class EMediaEvent;
 
+/** Playback parameters to update the media player. */ 
+struct FMovieSceneMediaPlaybackParams
+{
+	/**
+	 * Indicate if player is looping from the corresponding media section parameter.
+	 * We need to store this here to have it in HandleMediaPlayerEvent.
+	 */
+	bool bIsLooping = false;
+
+	/**
+	 * If specified, this is the playback time range (in player's time) calculated from the section and sequencer.
+	 * We need to store this here to have it in HandleMediaPlayerEvent.
+	 */
+	TRange<FTimespan> SectionTimeRange;
+
+	/**
+	 * Sequencer frame duration used for range tolerance calculations.
+	 */
+	FTimespan FrameDuration;
+};
+
 /**
  * Persistent data that's stored for each currently evaluating section.
  */
@@ -58,8 +79,11 @@ public:
 	 */
 	UE_API void SeekOnOpen(FTimespan Time);
 
-	/** Set up this persistent data object. */
+	UE_DEPRECATED(5.6, "Use new Setup function with the playback parameters")
 	UE_API void Setup(const TSharedPtr<FMovieSceneMediaPlayerStore>& InMediaPlayerStore, const UObject* InSection, UMediaPlayer* OverrideMediaPlayer, UObject* InPlayerProxy, int32 InProxyLayerIndex, int32 InProxyTextureIndex);
+	
+	/** Set up this persistent data object. */
+	UE_API void Setup(const TSharedPtr<FMovieSceneMediaPlayerStore>& InMediaPlayerStore, const UObject* InSection, UMediaPlayer* OverrideMediaPlayer, UObject* InPlayerProxy, int32 InProxyLayerIndex, int32 InProxyTextureIndex, const FMovieSceneMediaPlaybackParams& InPlaybackParams);
 
 	/**
 	 * Called from FMovieSceneMediaSectionTemplate::Initialize.
@@ -79,18 +103,6 @@ public:
 	 */
 	bool bIsAspectRatioSet;
 
-	/**
-	 * Indicate if player is looping from the corresponding media section parameter.
-	 * We need to store this here to have it in HandleMediaPlayerEvent.
-	 */
-	bool bIsLooping = false;
-
-	/**
-	 * If specified, this is the playback range calculated from the section and sequencer.
-	 * We need to store this here to have it in HandleMediaPlayerEvent.
-	 */
-	TRange<FTimespan> PlaybackRange;
-
 private:
 	/**
 	 * Does the work needed so we can use our proxy media texture.
@@ -136,6 +148,9 @@ private:
 
 	/** Owning Media Section. Used as persistent identifier for media player data. */
 	FObjectKey MediaSection;
+
+	/** Additional player parameters to set before the first seek. */
+	FMovieSceneMediaPlaybackParams PlaybackParams;
 };
 
 #undef UE_API
