[MediaCompositing] MediaTrack.UpdatePlaybackRange is now true by default.

diff --git a/Engine/Plugins/Media/MediaCompositing/Source/MediaCompositing/Private/MovieScene/MovieSceneMediaTemplate.cpp b/Engine/Plugins/Media/MediaCompositing/Source/MediaCompositing/Private/MovieScene/MovieSceneMediaTemplate.cpp
--- a/Engine/Plugins/Media/MediaCompositing/Source/MediaCompositing/Private/MovieScene/MovieSceneMediaTemplate.cpp
+++ b/Engine/Plugins/Media/MediaCompositing/Source/MediaCompositing/Private/MovieScene/MovieSceneMediaTemplate.cpp
@@ -29,14 +29,14 @@
 
 static TAutoConsoleVariable<bool> CVarUpdatePlaybackRange(
 	TEXT("MediaTrack.UpdatePlaybackRange"),
-	false, // Disabled by default until issues with Electra are ironed out.
+	true,
 	TEXT("Update Player PlaybackRange. Improves looping performance with better pre-caching."),
 	ECVF_Default);
 
 static TAutoConsoleVariable<bool> CVarAvoidSeekOnJump(
 	TEXT("MediaTrack.AvoidSeekOnJump"),
 	true,
-	TEXT("When the playhead jumps (on loop for instance), seek only if needed, i.e. if player doesn't have cached samples already."),
+	TEXT("When the play head jumps (on loop for instance), seek only if needed, i.e. if player doesn't have cached samples already."),
 	ECVF_Default);
 
 
