[Media Compositing] Fix Media Streams exposed to sequencer, and any media track with no texture.

Regression introduced in CL 42398244.
We need to discard the SampleQueue created for the Preroll if no managed texture use it.

diff --git a/Engine/Plugins/Media/MediaCompositing/Source/MediaCompositing/Private/MovieScene/MovieSceneMediaTemplate.cpp b/Engine/Plugins/Media/MediaCompositing/Source/MediaCompositing/Private/MovieScene/MovieSceneMediaTemplate.cpp
--- a/Engine/Plugins/Media/MediaCompositing/Source/MediaCompositing/Private/MovieScene/MovieSceneMediaTemplate.cpp
+++ b/Engine/Plugins/Media/MediaCompositing/Source/MediaCompositing/Private/MovieScene/MovieSceneMediaTemplate.cpp
@@ -743,6 +743,15 @@ void FMovieSceneMediaSectionTemplate::Initialize(const FMovieSceneEvaluationOper
 			Params.MediaTexture->SetMediaPlayer(nullptr);
 		}
 	}
+	else
+	{
+		// Make sure to discard the sample queue used for pre-roll if it isn't transferred to a media texture
+		// otherwise it will block the player by not consuming the samples.
+		if (IsEvaluating)
+		{
+			SectionData->TransferSampleQueue();
+		}
+	}
 
 	if (MediaPlayer->IsLooping() != Params.bLooping)
 	{
