[MediaCompositing] Fix block on range when moving the media section while playing.

The reason for this new problem is related to CL 40638157 for UE-231598 since now the player "survives" the rebuilding of the evaluation template. There was no logic to detect that it needed to seek in that particular condition.

Since the sequencer doesn't consider moving the section itself as a "jump". This doesn't trigger a seek. The evaluation template needs to determine if it was rebuilt and then check if a seek is needed. We do that by using a flag in the section data to know if it was evaluated before. If it is the first evaluation, and we have the player already playing, then we check if a seek is needed to avoid the player blocking.

diff --git a/Engine/Plugins/Media/MediaCompositing/Source/MediaCompositing/Private/MovieScene/MovieSceneMediaTemplate.cpp b/Engine/Plugins/Media/MediaCompositing/Source/MediaCompositing/Private/MovieScene/MovieSceneMediaTemplate.cpp
--- a/Engine/Plugins/Media/MediaCompositing/Source/MediaCompositing/Private/MovieScene/MovieSceneMediaTemplate.cpp
+++ b/Engine/Plugins/Media/MediaCompositing/Source/MediaCompositing/Private/MovieScene/MovieSceneMediaTemplate.cpp
@@ -111,7 +111,7 @@ namespace UE::MediaSectionBaseMediaUtils
 		}
 
 		// Otherwise, check if we have a sample available for the desired time, either in the sample queue or already processed.
-		return IsVideoSampleAvailableForTime(MediaPlayerFacade, InTime);
+		return !IsVideoSampleAvailableForTime(MediaPlayerFacade, InTime);
 	}
 
 	/**
@@ -470,7 +470,7 @@ struct FMediaSectionExecutionToken
 			if (!MediaPlayer->IsPlaying())
 			{
 				// If the player has pre-rolled, it is ready and doesn't need to seek.
-				if (!UE::MediaSectionBaseMediaUtils::ShouldSeekForTime(MediaPlayer, MediaTime))
+				if (ShouldSeekForTime(MediaPlayer, MediaTime))
 				{
 					SeekPlayer(Context, MediaPlayer, MediaTime);
 				}
@@ -489,7 +489,7 @@ struct FMediaSectionExecutionToken
 			else
 			{
 				// Avoid seek on jump (loop) if the player already has cached samples.
-				if (Context.HasJumped() && (!CVarAvoidSeekOnJump.GetValueOnGameThread() || !ShouldSeekForTime(MediaPlayer, MediaTime)))
+				if ((Context.HasJumped() || !SectionData.bHasBeenExecuted) && (!CVarAvoidSeekOnJump.GetValueOnGameThread() || ShouldSeekForTime(MediaPlayer, MediaTime)))
 				{
 					SeekPlayer(Context, MediaPlayer, MediaTime);
 				}
@@ -539,6 +539,9 @@ struct FMediaSectionExecutionToken
 			// (we always use the full current time for this, any adjustments to player timestamps are done internally)
 			SetPlayerBlockOnTimeRange(MediaPlayer, CurrentTime, GetPlaybackParams());
 		}
+		
+		// Mark the section data as having been evaluated.
+		SectionData.bHasBeenExecuted = true;
 	}
 
 private:
@@ -615,13 +618,14 @@ void FMovieSceneMediaSectionTemplate::Evaluate(const FMovieSceneEvaluationOperan
 		const float ProxyTextureBlend = MediaSection->EvaluateEasing(Context.GetTime());
 
 		#if MOVIESCENEMEDIATEMPLATE_TRACE_EVALUATION
-			GLog->Logf(ELogVerbosity::Log, TEXT("Evaluating frame %i+%f, FrameRate %i/%i, FrameTicks %d, FrameDurationTicks %d"),
+			GLog->Logf(ELogVerbosity::Log, TEXT("Evaluating (%s) frame %i+%f, FrameRate %i/%i, FrameTicks %lld, FrameDurationTicks %lld"),
+				*StaticEnum<EMovieScenePlayerStatus::Type>()->GetNameStringByValue(Context.GetStatus()),
 				Context.GetTime().GetFrame().Value,
 				Context.GetTime().GetSubFrame(),
 				FrameRate.Numerator,
 				FrameRate.Denominator,
-				FrameTicks,
-				FrameDurationTicks
+				FrameTime.GetTicks(),
+				UE::MediaSectionBaseMediaUtils::GetFrameDuration(Context).GetTicks()
 			);
 		#endif
 
diff --git a/Engine/Plugins/Media/MediaCompositing/Source/MediaCompositing/Public/MovieSceneMediaData.h b/Engine/Plugins/Media/MediaCompositing/Source/MediaCompositing/Public/MovieSceneMediaData.h
--- a/Engine/Plugins/Media/MediaCompositing/Source/MediaCompositing/Public/MovieSceneMediaData.h
+++ b/Engine/Plugins/Media/MediaCompositing/Source/MediaCompositing/Public/MovieSceneMediaData.h
@@ -103,6 +103,12 @@ public:
 	 */
 	bool bIsAspectRatioSet;
 
+	/**
+	 * Indicate if the section template has been executed already or not.
+	 * This is used to detect when a evaluation template was rebuilt and has potentially changed.
+	 */
+	bool bHasBeenExecuted = false;
+	
 private:
 	/**
 	 * Does the work needed so we can use our proxy media texture.
