[ImgMedia] EXR sequences: Fix loading of mips and tiles during sequencer media track pre-roll

Problem:
The media texture can't be assigned (i.e. used as a sample sink) to the player during the preroll because it is used by another player of overlapping playback section. But the player needs to be specified the media texture to properly preload the mips and tiles accordingly.

Implemented solution:
A new loading option is added to specify a "view texture" when opening the media source to allow the player to load only the mips and tiles needed for the given view earlier in the loading process.

diff --git a/Engine/Plugins/Media/ImgMedia/Source/ImgMedia/Private/Player/ImgMediaPlayer.cpp b/Engine/Plugins/Media/ImgMedia/Source/ImgMedia/Private/Player/ImgMediaPlayer.cpp
--- a/Engine/Plugins/Media/ImgMedia/Source/ImgMedia/Private/Player/ImgMediaPlayer.cpp
+++ b/Engine/Plugins/Media/ImgMedia/Source/ImgMedia/Private/Player/ImgMediaPlayer.cpp
@@ -236,6 +236,11 @@ bool FImgMediaPlayer::Open(const FString& Url, const IMediaOptions* Options, con
 		const FVariant* Var = InPlayerOptions ? InPlayerOptions->InternalCustomOptions.Find(InOptionName) : nullptr;
 		return Var && Var->GetType() == EVariantTypes::Float ? (double)Var->GetValue<float>() : Var && Var->GetType() == EVariantTypes::Double ? Var->GetValue<double>() : InDefaultValue;
 	};
+	auto GetStringOption = [](const FMediaPlayerOptions* InPlayerOptions, const FName& InOptionName) -> FString
+	{
+		const FVariant* Var = InPlayerOptions ? InPlayerOptions->InternalCustomOptions.Find(InOptionName) : nullptr;
+		return Var && Var->GetType() == EVariantTypes::String ? Var->GetValue<FString>() : FString();
+	};
 
 	// get frame rate override, if any
 	FFrameRate FrameRateOverride(0, 0);
@@ -244,6 +249,13 @@ bool FImgMediaPlayer::Open(const FString& Url, const IMediaOptions* Options, con
 	bool bFillGapsInSequence = true;
 	bool bIsSmartCacheEnabled = GetBoolOption(PlayerOptions, MediaPlayerOptionValues::ImgMediaSmartCacheEnabled(), false);
 	double SmartCacheTimeToLookAhead = GetDoubleOption(PlayerOptions, MediaPlayerOptionValues::ImgMediaSmartCacheTimeToLookAhead(), 0.0);
+
+	UMediaTexture* ViewTexture = nullptr;
+	{
+		FSoftObjectPath ViewTexturePath = GetStringOption(PlayerOptions, MediaPlayerOptionValues::ViewMediaTexture());
+		ViewTexture =  Cast<UMediaTexture>(ViewTexturePath.ResolveObject());
+	}
+	
 	if (Options != nullptr)
 	{
 		FrameRateOverride.Denominator = Options->GetMediaOption(ImgMedia::FrameRateOverrideDenonimatorOption, 0LL);
@@ -256,24 +268,31 @@ bool FImgMediaPlayer::Open(const FString& Url, const IMediaOptions* Options, con
 			MipMapInfo = StaticCastSharedPtr<FImgMediaMipMapInfo, IMediaOptions::FDataContainer, ESPMode::ThreadSafe>(DataContainer);
 			if (MipMapInfo.IsValid())
 			{
-				// Look through all the media textures we know about to find the one tied to this player
-				for (TWeakObjectPtr<UMediaTexture> TexturePtr : FMediaTextureTracker::Get().GetTextures())
+				if (ViewTexture)
+				{
+					MediaTextureWeakPtr = ViewTexture;
+					MipMapInfo->AddObjectsUsingThisMediaTexture(MediaTextureWeakPtr.Get());
+				}
+				else
 				{
-					UMediaTexture* Texture = TexturePtr.Get();
-					if (Texture != nullptr)
+					// Look through all the media textures we know about to find the one tied to this player
+					for (TWeakObjectPtr<UMediaTexture> TexturePtr : FMediaTextureTracker::Get().GetTextures())
 					{
-						// Does this match the player?
-						UMediaPlayer* MediaPlayer = Texture->GetMediaPlayer();
-						if (MediaPlayer != nullptr)
+						UMediaTexture* Texture = TexturePtr.Get();
+						if (Texture != nullptr)
 						{
-							TSharedPtr<IMediaPlayer, ESPMode::ThreadSafe> Player = MediaPlayer->GetPlayerFacade()->GetPlayer();
-							if (Player.IsValid())
+							// Does this match the player?
+							UMediaPlayer* MediaPlayer = Texture->GetMediaPlayer();
+							if (MediaPlayer != nullptr)
 							{
-								if (Player.Get() == this)
+								TSharedPtr<IMediaPlayer, ESPMode::ThreadSafe> Player = MediaPlayer->GetPlayerFacade()->GetPlayer();
+								if (Player.IsValid())
 								{
-									MediaTextureWeakPtr = TexturePtr;
-
-									MipMapInfo->AddObjectsUsingThisMediaTexture(MediaTextureWeakPtr.Get());
+									if (Player.Get() == this)
+									{
+										MediaTextureWeakPtr = TexturePtr;
+										MipMapInfo->AddObjectsUsingThisMediaTexture(MediaTextureWeakPtr.Get());
+									}
 								}
 							}
 						}
diff --git a/Engine/Plugins/Media/MediaCompositing/Source/MediaCompositing/Private/MovieScene/MovieSceneMediaData.cpp b/Engine/Plugins/Media/MediaCompositing/Source/MediaCompositing/Private/MovieScene/MovieSceneMediaData.cpp
--- a/Engine/Plugins/Media/MediaCompositing/Source/MediaCompositing/Private/MovieScene/MovieSceneMediaData.cpp
+++ b/Engine/Plugins/Media/MediaCompositing/Source/MediaCompositing/Private/MovieScene/MovieSceneMediaData.cpp
@@ -131,7 +131,7 @@ TSharedPtr<FMediaTextureSampleQueue> FMovieSceneMediaData::TransferSampleQueue()
 	return OutSampleQueue;
 }
 
-void FMovieSceneMediaData::StartUsingProxyMediaTexture()
+UMediaTexture* FMovieSceneMediaData::GetProxyMediaTexture()
 {
 	if (PlayerProxy != nullptr)
 	{
@@ -143,10 +143,15 @@ void FMovieSceneMediaData::StartUsingProxyMediaTexture()
 				ProxyMediaTexture = PlayerProxyInterface->ProxyGetMediaTexture(ProxyLayerIndex, ProxyTextureIndex);
 			}
 		}
-		if (ProxyMediaTexture != nullptr)
-		{
-			ProxyMediaTexture->SetMediaPlayer(MediaPlayer, TransferSampleQueue());
-		}
+	}
+	return ProxyMediaTexture.Get();
+}
+
+void FMovieSceneMediaData::StartUsingProxyMediaTexture()
+{
+	if (UMediaTexture* MediaTexture = GetProxyMediaTexture())
+	{
+		MediaTexture->SetMediaPlayer(MediaPlayer, TransferSampleQueue());
 	}
 }
 
diff --git a/Engine/Plugins/Media/MediaCompositing/Source/MediaCompositing/Private/MovieScene/MovieSceneMediaTemplate.cpp b/Engine/Plugins/Media/MediaCompositing/Source/MediaCompositing/Private/MovieScene/MovieSceneMediaTemplate.cpp
--- a/Engine/Plugins/Media/MediaCompositing/Source/MediaCompositing/Private/MovieScene/MovieSceneMediaTemplate.cpp
+++ b/Engine/Plugins/Media/MediaCompositing/Source/MediaCompositing/Private/MovieScene/MovieSceneMediaTemplate.cpp
@@ -268,6 +268,32 @@ struct FMediaSectionBaseExecutionToken
 	/** Get the section's playback parameters. */
 	const FMovieSceneMediaPlaybackParams& GetPlaybackParams() const { return PlaybackParams; }
 
+	/**
+	 * Utility function to prepare the media player options for opening the media source.
+	 */
+	FMediaPlayerOptions MakeMediaPlayerOptions(const IMediaPlayerProxyInterface* InPlayerProxyInterface, FMovieSceneMediaData& InSectionData) const
+	{
+		FMediaPlayerOptions Options;
+		Options.SetAllAsOptional();
+		Options.InternalCustomOptions.Emplace(MediaPlayerOptionValues::Environment(), MediaPlayerOptionValues::Environment_Sequencer());
+
+		if (InPlayerProxyInterface != nullptr)
+		{
+			// Set cache settings.
+			const FMediaSourceCacheSettings& CacheSettings = InPlayerProxyInterface->GetCacheSettings();
+			Options.InternalCustomOptions.Emplace(MediaPlayerOptionValues::ImgMediaSmartCacheEnabled(), FVariant(CacheSettings.bOverride));
+			Options.InternalCustomOptions.Emplace(MediaPlayerOptionValues::ImgMediaSmartCacheTimeToLookAhead(), FVariant(CacheSettings.TimeToLookAhead));
+
+			// Set the view texture for proper mips and tiles loading during pre-roll.
+			// This is only done if we have a PlayerProxyInterface because it is the only case with associated visibility geometry (ex MediaPlate).
+			if (UMediaTexture* ViewTexture = InSectionData.GetProxyMediaTexture())
+			{
+				Options.InternalCustomOptions.Emplace(MediaPlayerOptionValues::ViewMediaTexture(), FVariant(FSoftObjectPath(ViewTexture).ToString()));
+			}
+		}
+		return Options;
+	}
+
 private:
 	UMediaSource* BaseMediaSource;
 	FMovieSceneObjectBindingID MediaSourceProxy;
@@ -303,17 +329,7 @@ struct FMediaSectionPreRollExecutionToken
 		{
 			SectionData.SeekOnOpen(StartTime);
 
-			FMediaPlayerOptions Options;
-			Options.SetAllAsOptional();
-			Options.InternalCustomOptions.Emplace(MediaPlayerOptionValues::Environment(), MediaPlayerOptionValues::Environment_Sequencer());
-			if (PlayerProxyInterface != nullptr)
-			{
-				// Set cache settings.
-				const FMediaSourceCacheSettings& CacheSettings = PlayerProxyInterface->GetCacheSettings();
-				Options.InternalCustomOptions.Emplace(MediaPlayerOptionValues::ImgMediaSmartCacheEnabled(), FVariant(CacheSettings.bOverride));
-				Options.InternalCustomOptions.Emplace(MediaPlayerOptionValues::ImgMediaSmartCacheTimeToLookAhead(), FVariant(CacheSettings.TimeToLookAhead));
-			}
-
+			FMediaPlayerOptions Options = MakeMediaPlayerOptions(PlayerProxyInterface, SectionData);
 			MediaPlayer->OpenSourceWithOptions(MediaSource, Options);
 			return;
 		}
@@ -420,15 +436,7 @@ struct FMediaSectionExecutionToken
 		{
 			SectionData.SeekOnOpen(CurrentTime);
 
-			FMediaPlayerOptions Options;
-			Options.SetAllAsOptional();
-			Options.InternalCustomOptions.Emplace(MediaPlayerOptionValues::Environment(), MediaPlayerOptionValues::Environment_Sequencer());
-			if (PlayerProxyInterface != nullptr)
-			{
-				const FMediaSourceCacheSettings& CacheSettings = PlayerProxyInterface->GetCacheSettings();
-				Options.InternalCustomOptions.Emplace(MediaPlayerOptionValues::ImgMediaSmartCacheEnabled(), FVariant(CacheSettings.bOverride));
-				Options.InternalCustomOptions.Emplace(MediaPlayerOptionValues::ImgMediaSmartCacheTimeToLookAhead(), FVariant(CacheSettings.TimeToLookAhead));
-			}
+			FMediaPlayerOptions Options = MakeMediaPlayerOptions(PlayerProxyInterface, SectionData);
 
 			// Setup an initial blocking range - MediaFramework will block (even through the opening process) in its next tick...
 			SetPlayerBlockOnTimeRange(MediaPlayer, CurrentTime, GetPlaybackParams());
diff --git a/Engine/Plugins/Media/MediaCompositing/Source/MediaCompositing/Public/MovieSceneMediaData.h b/Engine/Plugins/Media/MediaCompositing/Source/MediaCompositing/Public/MovieSceneMediaData.h
--- a/Engine/Plugins/Media/MediaCompositing/Source/MediaCompositing/Public/MovieSceneMediaData.h
+++ b/Engine/Plugins/Media/MediaCompositing/Source/MediaCompositing/Public/MovieSceneMediaData.h
@@ -97,6 +97,9 @@ public:
 
 	/** Retrieve the sample queue and release the ownership. */
 	UE_API TSharedPtr<FMediaTextureSampleQueue> TransferSampleQueue();
+
+	/** Get the proxy media texture. */
+	UE_API UMediaTexture* GetProxyMediaTexture();
 	
 	/**
 	 * Stores if the aspect ratio has been set yet.
diff --git a/Engine/Source/Runtime/MediaUtils/Public/MediaPlayerOptions.h b/Engine/Source/Runtime/MediaUtils/Public/MediaPlayerOptions.h
--- a/Engine/Source/Runtime/MediaUtils/Public/MediaPlayerOptions.h
+++ b/Engine/Source/Runtime/MediaUtils/Public/MediaPlayerOptions.h
@@ -202,6 +202,13 @@ namespace MediaPlayerOptionValues
 	inline const FName& ImgMediaSmartCacheTimeToLookAhead()
 	{ static FName OptName(TEXT("ImgMediaSmartCacheTimeToLookAhead")); return OptName; }
 
+	/**
+	 * Option to specify a view media texture (as a soft object path string)
+	 * to allow players supporting mips and tiles to preload only the required set for the given texture.
+	 */
+	inline const FName& ViewMediaTexture()
+	{ static FName OptName(TEXT("ViewMediaTexture")); return OptName; }
+
 	/**
 	 * Requests parsing timecode meta data if available.
 	 */
