[MediaCompositing] Fix Block On Range time out when "frame accurate" scrubbing is enabled and scrubbing outside of the playback range.

The "frame accurate" scrubbing is implemented using the "block on range" player feature. Clamping the time to the player's playback range is not correct when scrubbing outside of the sequencer's playback range.

diff --git a/Engine/Plugins/Media/MediaCompositing/Source/MediaCompositing/Private/MovieScene/MovieSceneMediaData.cpp b/Engine/Plugins/Media/MediaCompositing/Source/MediaCompositing/Private/MovieScene/MovieSceneMediaData.cpp
--- a/Engine/Plugins/Media/MediaCompositing/Source/MediaCompositing/Private/MovieScene/MovieSceneMediaData.cpp
+++ b/Engine/Plugins/Media/MediaCompositing/Source/MediaCompositing/Private/MovieScene/MovieSceneMediaData.cpp
@@ -212,19 +212,19 @@ void FMovieSceneMediaData::HandleMediaPlayerEvent(EMediaEvent Event)
 	}
 
 	const FTimespan ClampTolerance = PlaybackParams.FrameDuration * 0.5f;
+	MediaTime = FMath::Clamp(MediaTime, FTimespan::Zero(), Duration - ClampTolerance);
 
 	if (!PlaybackParams.SectionTimeRange.IsEmpty() && MediaPlayer->SupportsPlaybackTimeRange())
 	{
 		using namespace UE::MovieSceneMediaPlayerUtils;
 		const TRange<FTimespan> AdjustedRange = AdjustPlaybackTimeRange(PlaybackParams.SectionTimeRange, MediaPlayer, PlaybackParams.FrameDuration);
-		SetPlayerPlaybackTimeRange(MediaPlayer, AdjustedRange);
 
-		// Ensure we are seeking within the specified playback range.
-		MediaTime = FMath::Clamp(MediaTime, AdjustedRange.GetLowerBoundValue(), AdjustedRange.GetUpperBoundValue() - ClampTolerance);
-	}
-	else
-	{
-		MediaTime = FMath::Clamp(MediaTime, FTimespan::Zero(), Duration - ClampTolerance);
+		// We can only set the player's playback time range if the requested seek time is contained within.
+		// It is possible the media is opened by scrubbing or stepping outside the playback range.
+		if (!AdjustedRange.IsEmpty() && AdjustedRange.Contains(MediaTime))
+		{
+			SetPlayerPlaybackTimeRange(MediaPlayer, AdjustedRange);
+		}
 	}
 
 	MediaPlayer->SetRate(0.0f);
diff --git a/Engine/Plugins/Media/MediaCompositing/Source/MediaCompositing/Private/MovieScene/MovieSceneMediaTemplate.cpp b/Engine/Plugins/Media/MediaCompositing/Source/MediaCompositing/Private/MovieScene/MovieSceneMediaTemplate.cpp
--- a/Engine/Plugins/Media/MediaCompositing/Source/MediaCompositing/Private/MovieScene/MovieSceneMediaTemplate.cpp
+++ b/Engine/Plugins/Media/MediaCompositing/Source/MediaCompositing/Private/MovieScene/MovieSceneMediaTemplate.cpp
@@ -202,12 +202,11 @@ namespace UE::MediaSectionBaseMediaUtils
 	 */
 	void SetPlayerBlockOnTimeRange(UMediaPlayer* InMediaPlayer, const FTimespan& InCurrentTime, const FMovieSceneMediaPlaybackParams& InPlaybackParams)
 	{
-		FTimespan CurrentTime = UE::MovieSceneMediaPlayerUtils::ClampTimeToPlaybackRange(InCurrentTime, InMediaPlayer, InPlaybackParams);
-		InMediaPlayer->SetBlockOnTimeRange(TRange<FTimespan>(CurrentTime, CurrentTime + InPlaybackParams.FrameDuration));
+		InMediaPlayer->SetBlockOnTimeRange(TRange<FTimespan>(InCurrentTime, InCurrentTime + InPlaybackParams.FrameDuration));
 	}
 
 	/** Calculate the timespan of a frame duration given the sequencer's framerate */
-	FTimespan GetFrameDuration(const FMovieSceneContext& InContext)
+	FTimespan GetFrameDuration(const FMovieSceneContext& InContext) 
 	{
 		const FFrameRate FrameRate = InContext.GetFrameRate();
 		// With zero-length frames (which can occur occasionally), we use the fixed frame time, matching previous behavior.
@@ -372,7 +371,15 @@ struct FMediaSectionExecutionToken
 				UE::MovieSceneMediaPlayerUtils::SetPlayerPlaybackTimeRange(InMediaPlayer, TRange<FTimespan>::Empty());
 			}
 
-			InMediaPlayer->Scrub(InMediaTime);
+			if (!GetBaseMediaSourceOptions().bSynchronousScrubbing)
+			{
+				InMediaPlayer->Scrub(InMediaTime);
+			}
+			else
+			{
+				// Force seek here to avoid issues with block on range.
+				InMediaPlayer->Seek(InMediaTime);
+			}
 		}
 		else
 		{
@@ -642,7 +649,22 @@ void FMovieSceneMediaSectionTemplate::Evaluate(const FMovieSceneEvaluationOperan
 		FMediaSectionBaseMediaSourceOptions Options;
 		Options.CacheSettings = Params.CacheSettings;
 		Options.bSynchronousScrubbing = MediaTrack ? MediaTrack->bSynchronousScrubbing : false;
-		const FMovieSceneMediaPlaybackParams PlaybackParams = UE::MediaSectionBaseMediaUtils::MakePlaybackParams(Context, Params, MediaSection);
+		FMovieSceneMediaPlaybackParams PlaybackParams = UE::MediaSectionBaseMediaUtils::MakePlaybackParams(Context, Params, MediaSection);
+
+		// Scrubbing or stepping outside the playback range is allowed by the sequencer,
+		// in which case we have to reset the player's playback range.
+		if (Context.GetStatus() == EMovieScenePlayerStatus::Scrubbing
+			|| Context.GetStatus() == EMovieScenePlayerStatus::Stopped
+			|| Context.GetStatus() == EMovieScenePlayerStatus::Stepping)
+		{
+			if (!PlaybackParams.SectionTimeRange.IsEmpty() && !PlaybackParams.SectionTimeRange.Contains(FrameTime))
+			{
+				// Note: this only resets the playback time range for the current evaluation.
+				// It will propagate to the player, but not the persistent section data (FMovieSceneMediaData).
+				PlaybackParams.SectionTimeRange = TRange<FTimespan>::Empty();
+			}
+		}
+		
 		ExecutionTokens.Add(FMediaSectionExecutionToken(MediaSource, Options, PlaybackParams, Params.MediaSourceProxy, Params.MediaSourceProxyIndex, ProxyTextureBlend, bCanPlayerBeOpen, FrameTime));
 	}
 }
