[MediaCompositing] Fix PostRoll and fix issue with playback range boundary wrap around condition that affects Protron primarily.

diff --git a/Engine/Plugins/Media/MediaCompositing/Source/MediaCompositing/Private/MovieScene/MovieSceneMediaTemplate.cpp b/Engine/Plugins/Media/MediaCompositing/Source/MediaCompositing/Private/MovieScene/MovieSceneMediaTemplate.cpp
--- a/Engine/Plugins/Media/MediaCompositing/Source/MediaCompositing/Private/MovieScene/MovieSceneMediaTemplate.cpp
+++ b/Engine/Plugins/Media/MediaCompositing/Source/MediaCompositing/Private/MovieScene/MovieSceneMediaTemplate.cpp
@@ -202,7 +202,32 @@ namespace UE::MediaSectionBaseMediaUtils
 	 */
 	void SetPlayerBlockOnTimeRange(UMediaPlayer* InMediaPlayer, const FTimespan& InCurrentTime, const FMovieSceneMediaPlaybackParams& InPlaybackParams)
 	{
-		InMediaPlayer->SetBlockOnTimeRange(TRange<FTimespan>(InCurrentTime, InCurrentTime + InPlaybackParams.FrameDuration));
+		FTimespan RangeLowerBound = InCurrentTime;
+		FTimespan RangeUpperBound = InCurrentTime + InPlaybackParams.FrameDuration;
+
+		// Player Facade currently (as of 5.6) has issues when the BlockOnTimeRange is partially outside the player's active range and especially when
+		// it doesn't start at zero. The wrap around code for the boundaries has inconsistencies (that will need to be fixed).
+		// As a temporary workaround for that, we will clamp the BlockOnTimeRange to the section's active range to make
+		// sure to avoid triggering any of the Player Facade internal wrap around boundary code.
+		if (InMediaPlayer->SupportsPlaybackTimeRange() && !InPlaybackParams.SectionTimeRange.IsEmpty())
+		{
+			TRange<FTimespan> ActiveRange = InPlaybackParams.SectionTimeRange;
+		
+			// BlockOnTimeRange is spanning the upper limit of the active range
+			if (ActiveRange.Contains(RangeLowerBound) && RangeUpperBound > ActiveRange.GetUpperBoundValue())
+			{
+				RangeUpperBound = ActiveRange.GetUpperBoundValue() - FTimespan(1); // One tick inside the range to avoid wrap around.
+				RangeLowerBound = ActiveRange.GetUpperBoundValue() - InPlaybackParams.FrameDuration;
+			}
+			// BlockOnTimeRange is spanning the lower limit of the active range
+			else if (ActiveRange.Contains(RangeUpperBound) && RangeLowerBound < ActiveRange.GetLowerBoundValue())
+			{
+				RangeLowerBound = ActiveRange.GetLowerBoundValue() + FTimespan(1); // One tick inside the range to avoid wrap around.
+				RangeUpperBound = ActiveRange.GetLowerBoundValue() + InPlaybackParams.FrameDuration;
+			}
+		}
+		
+		InMediaPlayer->SetBlockOnTimeRange(TRange<FTimespan>(RangeLowerBound, RangeUpperBound));
 	}
 
 	/** Calculate the timespan of a frame duration given the sequencer's framerate */
@@ -225,6 +250,14 @@ namespace UE::MediaSectionBaseMediaUtils
 		Params.bIsLooping = InParams.bLooping;
 		return Params;
 	}
+
+	/** Returns true if the player is currently closed. */
+	bool IsPlayerClosed(const UMediaPlayer& InMediaPlayer)
+	{
+		// Fixme: IsClosed() returns false when the player was never opened.
+		// We check the internal player name as a workaround for that (i.e. no player name means it was never opened and thus is currently closed). 
+		return InMediaPlayer.GetPlayerName() == NAME_None || InMediaPlayer.IsClosed();
+	}
 }
 
 /** Base struct for execution tokens. */
@@ -336,7 +369,7 @@ struct FMediaSectionPreRollExecutionToken
 		const bool bMoveToNewTime = Context.GetStatus() != EMovieScenePlayerStatus::Playing || (Context.GetStatus() == EMovieScenePlayerStatus::Playing && Context.HasJumped());
 		if (bMoveToNewTime)
 		{
-			if (MediaPlayer->GetRate() != 0.0f)	// avoids Electra spamming when scrubbing in the preroll.
+			if (MediaPlayer->GetRate() != 0.0f)	// avoids Electra spamming "SetRate" when scrubbing in the preroll.
 			{
 				MediaPlayer->SetRate(0.0f);
 			}
@@ -350,6 +383,29 @@ private:
 	FTimespan StartTime;
 };
 
+struct FMediaSectionPostRollExecutionToken
+	: FMediaSectionBaseExecutionToken
+{
+	FMediaSectionPostRollExecutionToken(UMediaSource* InMediaSource, const FMediaSectionBaseMediaSourceOptions& InMediaSourceOptions, const FMovieSceneMediaPlaybackParams& InPlaybackParams, FMovieSceneObjectBindingID InMediaSourceProxy, int32 InMediaSourceProxyIndex)
+		: FMediaSectionBaseExecutionToken(InMediaSource, InMediaSourceOptions, InPlaybackParams, InMediaSourceProxy, InMediaSourceProxyIndex)
+	{ }
+
+	virtual void Execute(const FMovieSceneContext& Context, const FMovieSceneEvaluationOperand& Operand, FPersistentEvaluationData& PersistentData, IMovieScenePlayer& Player) override
+	{
+		FMovieSceneMediaData& SectionData = PersistentData.GetSectionData<FMovieSceneMediaData>();
+		UMediaPlayer* MediaPlayer = SectionData.GetMediaPlayer();
+
+		if (MediaPlayer == nullptr || UE::MediaSectionBaseMediaUtils::IsPlayerClosed(*MediaPlayer))
+		{
+			return;
+		}
+
+		if (MediaPlayer->GetRate() != 0.0f)	// avoids Electra spamming "SetRate" when scrubbing in post-roll.
+		{
+			MediaPlayer->SetRate(0.0f);
+		}
+	}
+};
 
 struct FMediaSectionExecutionToken
 	: FMediaSectionBaseExecutionToken
@@ -421,18 +477,18 @@ struct FMediaSectionExecutionToken
 			}
 		}
 
+		using namespace UE::MediaSectionBaseMediaUtils;
+
 		// Can we be open?
 		if (bCanPlayerBeOpen == false)
 		{
-			if (MediaPlayer->IsClosed() == false)
+			if (!IsPlayerClosed(*MediaPlayer))
 			{
 				MediaPlayer->Close();
 			}
 			return;
 		}
 
-		using namespace UE::MediaSectionBaseMediaUtils;
-
 		// Check if the cache options have changed mid playback.
 		//const FMediaSourceCacheSettings CurrentCacheSettings = GetBaseMediaSourceOptions().CacheSettings;
 		const FMediaSourceCacheSettings CurrentCacheSettings = GetCurrentCacheSettingsFromPlayer(MediaPlayer);
@@ -602,7 +658,7 @@ FMovieSceneMediaSectionTemplate::FMovieSceneMediaSectionTemplate(const UMovieSce
 void FMovieSceneMediaSectionTemplate::Evaluate(const FMovieSceneEvaluationOperand& Operand, const FMovieSceneContext& Context, const FPersistentEvaluationData& PersistentData, FMovieSceneExecutionTokens& ExecutionTokens) const
 {
 	UMediaSource* MediaSource = Params.MediaSource;
-	if (((MediaSource == nullptr) && (Params.MediaSourceProxy.IsValid() == false)) || Context.IsPostRoll())
+	if (((MediaSource == nullptr) && (Params.MediaSourceProxy.IsValid() == false)))
 	{
 		return;
 	}
@@ -624,6 +680,14 @@ void FMovieSceneMediaSectionTemplate::Evaluate(const FMovieSceneEvaluationOperan
 		const FMovieSceneMediaPlaybackParams PlaybackParams = UE::MediaSectionBaseMediaUtils::MakePlaybackParams(Context, Params, MediaSection);
 		ExecutionTokens.Add(FMediaSectionPreRollExecutionToken(MediaSource, Options, PlaybackParams, Params.MediaSourceProxy, Params.MediaSourceProxyIndex, StartTime));
 	}
+	else if (Context.IsPostRoll() && bCanPlayerBeOpen)
+	{
+		FMediaSectionBaseMediaSourceOptions Options;
+		Options.CacheSettings = Params.CacheSettings;
+		Options.bSynchronousScrubbing = false; // not relevant during post-roll
+		const FMovieSceneMediaPlaybackParams PlaybackParams = UE::MediaSectionBaseMediaUtils::MakePlaybackParams(Context, Params, MediaSection);
+		ExecutionTokens.Add(FMediaSectionPostRollExecutionToken(MediaSource, Options, PlaybackParams, Params.MediaSourceProxy, Params.MediaSourceProxyIndex));
+	}
 	else if (!Context.IsPostRoll() && (Context.GetTime().FrameNumber < Params.SectionEndFrame))
 	{
 		const FFrameRate FrameRate = Context.GetFrameRate();
@@ -787,7 +851,7 @@ void FMovieSceneMediaSectionTemplate::Initialize(const FMovieSceneEvaluationOper
 		}
 	}
 
-	if (MediaPlayer->IsLooping() != Params.bLooping)
+	if (!UE::MediaSectionBaseMediaUtils::IsPlayerClosed(*MediaPlayer) && MediaPlayer->IsLooping() != Params.bLooping)
 	{
 		MediaPlayer->SetLooping(Params.bLooping);
 	}
