Media: Minimize blocking on open while scrubbing.

Fixes partial regression of recent scrubbing improvements caused by CL 42232969 (though this change needs to be preserved to support the synchronous scrubbing option).

diff --git a/Engine/Plugins/Media/MediaCompositing/Source/MediaCompositing/Private/MovieScene/MovieSceneMediaTemplate.cpp b/Engine/Plugins/Media/MediaCompositing/Source/MediaCompositing/Private/MovieScene/MovieSceneMediaTemplate.cpp
--- a/Engine/Plugins/Media/MediaCompositing/Source/MediaCompositing/Private/MovieScene/MovieSceneMediaTemplate.cpp
+++ b/Engine/Plugins/Media/MediaCompositing/Source/MediaCompositing/Private/MovieScene/MovieSceneMediaTemplate.cpp
@@ -493,6 +493,7 @@ struct FMediaSectionExecutionToken
 		//const FMediaSourceCacheSettings CurrentCacheSettings = GetBaseMediaSourceOptions().CacheSettings;
 		const FMediaSourceCacheSettings CurrentCacheSettings = GetCurrentCacheSettingsFromPlayer(MediaPlayer);
 		const bool bCacheSettingsChanged = PlayerProxyInterface ? CurrentCacheSettings != PlayerProxyInterface->GetCacheSettings() : false;
+		const bool bEnableBlockingRequests = (Context.GetStatus() != EMovieScenePlayerStatus::Scrubbing && Context.GetStatus() != EMovieScenePlayerStatus::Stopped) || GetBaseMediaSourceOptions().bSynchronousScrubbing;
 
 		// open the media source if necessary
 		if (MediaPlayer->GetUrl() != MediaSource->GetUrl() || bCacheSettingsChanged)
@@ -501,8 +502,11 @@ struct FMediaSectionExecutionToken
 
 			FMediaPlayerOptions Options = MakeMediaPlayerOptions(PlayerProxyInterface, SectionData);
 
-			// Setup an initial blocking range - MediaFramework will block (even through the opening process) in its next tick...
-			SetPlayerBlockOnTimeRange(MediaPlayer, CurrentTime, GetPlaybackParams());
+			if (bEnableBlockingRequests)
+			{
+				// Setup an initial blocking range - MediaFramework will block (even through the opening process) in its next tick...
+				SetPlayerBlockOnTimeRange(MediaPlayer, CurrentTime, GetPlaybackParams());
+			}
 			MediaPlayer->OpenSourceWithOptions(MediaSource, Options);
 			return;
 		}
@@ -512,7 +516,10 @@ struct FMediaSectionExecutionToken
 		if (MediaPlayer->IsPreparing())
 		{
 			SectionData.SeekOnOpen(CurrentTime);
-			SetPlayerBlockOnTimeRange(MediaPlayer, CurrentTime, GetPlaybackParams());
+			if (bEnableBlockingRequests)
+			{
+				SetPlayerBlockOnTimeRange(MediaPlayer, CurrentTime, GetPlaybackParams());
+			}
 			return;
 		}
 
@@ -598,18 +605,17 @@ struct FMediaSectionExecutionToken
 			SeekPlayer(Context, MediaPlayer, MediaTime);
 		}
 
-		if ((Context.GetStatus() == EMovieScenePlayerStatus::Scrubbing || Context.GetStatus() == EMovieScenePlayerStatus::Stopped)
-			&& GetBaseMediaSourceOptions().bSynchronousScrubbing == false)
-		{
-			// When scrubbing, seek requests are non-blocking.
-			MediaPlayer->SetBlockOnTimeRange(TRange<FTimespan>::Empty());
-		}
-		else
+		if (bEnableBlockingRequests)
 		{
 			// Set blocking range / time-range to display
 			// (we always use the full current time for this, any adjustments to player timestamps are done internally)
 			SetPlayerBlockOnTimeRange(MediaPlayer, CurrentTime, GetPlaybackParams());
 		}
+		else
+		{
+			// When scrubbing, seek requests are non-blocking.
+			MediaPlayer->SetBlockOnTimeRange(TRange<FTimespan>::Empty());
+		}
 		
 		// Mark the section data as having been evaluated.
 		SectionData.bHasBeenExecuted = true;
