[MediaCompositing] Improve Media Track support for spawnable Media Plate.

When the media plate is spawnable, different editor operations will re-spawn the media plate. For instance, either saving or changing a track parameter will respawn the media plate. When that happens while the playhead is in the media section, some of the media section "proxied" objects (player proxy and media texture) will become stale. We need to detect that case and refresh the player proxy from the updated bindings.

For now, we are checking that the player proxy pointer becomes stale in the media section data. This is our indicator that it needs to be refreshed.

diff --git a/Engine/Plugins/Media/MediaCompositing/Source/MediaCompositing/Private/MovieScene/MovieSceneMediaData.cpp b/Engine/Plugins/Media/MediaCompositing/Source/MediaCompositing/Private/MovieScene/MovieSceneMediaData.cpp
--- a/Engine/Plugins/Media/MediaCompositing/Source/MediaCompositing/Private/MovieScene/MovieSceneMediaData.cpp
+++ b/Engine/Plugins/Media/MediaCompositing/Source/MediaCompositing/Private/MovieScene/MovieSceneMediaData.cpp
@@ -93,6 +93,16 @@ void FMovieSceneMediaData::Setup(const TSharedPtr<FMovieSceneMediaPlayerStore>&
 	ProxyTextureIndex = InProxyTextureIndex;
 	PlaybackParams = InPlaybackParams;
 
+	UpdatePlayerProxy(InPlayerProxy);
+
+	SampleQueue = MakeShared<FMediaTextureSampleQueue>();
+	MediaPlayer->GetPlayerFacade()->AddVideoSampleSink(SampleQueue.ToSharedRef());
+}
+
+void FMovieSceneMediaData::UpdatePlayerProxy(UObject* InPlayerProxy)
+{
+	UObject* PreviousPlayerProxy = PlayerProxy.Get();
+
 	// Do we have a valid proxy object?
 	if ((InPlayerProxy != nullptr) && (InPlayerProxy->Implements<UMediaPlayerProxyInterface>()))
 	{
@@ -103,8 +113,12 @@ void FMovieSceneMediaData::Setup(const TSharedPtr<FMovieSceneMediaPlayerStore>&
 		PlayerProxy.Reset();
 	}
 
-	SampleQueue = MakeShared<FMediaTextureSampleQueue>();
-	MediaPlayer->GetPlayerFacade()->AddVideoSampleSink(SampleQueue.ToSharedRef());
+	// If the player proxy is changed, we need to reset the cached proxied media texture so it is
+	// refreshed from the new proxy in the next evaluation.
+	if ((PlayerProxy == nullptr) || (PlayerProxy != PreviousPlayerProxy))
+	{
+		ProxyMediaTexture.Reset();
+	}
 }
 
 void FMovieSceneMediaData::Initialize(bool bIsEvaluating)
diff --git a/Engine/Plugins/Media/MediaCompositing/Source/MediaCompositing/Private/MovieScene/MovieSceneMediaTemplate.cpp b/Engine/Plugins/Media/MediaCompositing/Source/MediaCompositing/Private/MovieScene/MovieSceneMediaTemplate.cpp
--- a/Engine/Plugins/Media/MediaCompositing/Source/MediaCompositing/Private/MovieScene/MovieSceneMediaTemplate.cpp
+++ b/Engine/Plugins/Media/MediaCompositing/Source/MediaCompositing/Private/MovieScene/MovieSceneMediaTemplate.cpp
@@ -258,6 +258,41 @@ namespace UE::MediaSectionBaseMediaUtils
 		// We check the internal player name as a workaround for that (i.e. no player name means it was never opened and thus is currently closed). 
 		return InMediaPlayer.GetPlayerName() == NAME_None || InMediaPlayer.IsClosed();
 	}
+
+	/**
+	 * Utility function to find the player proxy interface object (that implements UMediaPlayerProxyInterface)
+	 * from the sequencer binding operand. 
+	 */
+	UObject* FindMediaPlayerProxyInterfaceObject(const FMovieSceneEvaluationOperand& InOperand, IMovieScenePlayer& InScenePlayer)
+	{
+		if (InOperand.ObjectBindingID.IsValid())
+		{
+			if (IMediaAssetsModule* MediaAssetsModule = FModuleManager::LoadModulePtr<IMediaAssetsModule>("MediaAssets"))
+			{
+				// Iterate over all bound objects for this operand.
+				for (TWeakObjectPtr<> WeakObject : InScenePlayer.FindBoundObjects(InOperand))
+				{
+					if (UObject* BoundObject = WeakObject.Get())
+					{
+						// Use the media asset module to resolve to a player proxy object.
+						UObject* PlayerProxy = nullptr;
+						MediaAssetsModule->GetPlayerFromObject(BoundObject, PlayerProxy);
+
+						// Remark: we are not interested in the media player itself. The media section will create and manage its own player.
+						// We are only interested in the player proxy interface, which is used for querying the media texture
+						// and other rendering parameters.
+
+						// If we have found a player proxy interface object, we can stop searching.
+						if (PlayerProxy && PlayerProxy->Implements<UMediaPlayerProxyInterface>())
+						{
+							return PlayerProxy;
+						}
+					}
+				}
+			}
+		}
+		return nullptr;
+	}
 }
 
 /** Base struct for execution tokens. */
@@ -759,29 +794,12 @@ void FMovieSceneMediaSectionTemplate::Initialize(const FMovieSceneEvaluationOper
 			ProxyLayerIndex = MediaSection->GetRowIndex();
 		}
 
-		// Are we overriding the media player?
-		UMediaPlayer* MediaPlayer = Params.MediaPlayer;
-		UObject* PlayerProxy = nullptr;
-		if (MediaPlayer == nullptr)
+		// If the media section doesn't specify an external media player, we try to find the
+		// media player proxy from the bound objects.
+		UObject* MediaPlayerProxy = nullptr;
+		if (Params.MediaPlayer == nullptr)
 		{
-			// Nope... do we have an object binding?
-			if (Operand.ObjectBindingID.IsValid())
-			{
-				// Yes. Get the media player from the object.
-				IMediaAssetsModule* MediaAssetsModule = FModuleManager::LoadModulePtr<IMediaAssetsModule>("MediaAssets");
-				if (MediaAssetsModule != nullptr)
-				{
-					for (TWeakObjectPtr<> WeakObject : Player.FindBoundObjects(Operand))
-					{
-						UObject* BoundObject = WeakObject.Get();
-						if (BoundObject != nullptr)
-						{
-							MediaAssetsModule->GetPlayerFromObject(BoundObject, PlayerProxy);
-							break;
-						}
-					}
-				}
-			}
+			MediaPlayerProxy = UE::MediaSectionBaseMediaUtils::FindMediaPlayerProxyInterfaceObject(Operand, Player);
 		}
 
 		// Add section data.
@@ -790,7 +808,7 @@ void FMovieSceneMediaSectionTemplate::Initialize(const FMovieSceneEvaluationOper
 		// Prepare these values in the section data to have them set in the player when it is opened (as early as possible before the first seek).
 		FMovieSceneMediaPlaybackParams PlaybackParams = UE::MediaSectionBaseMediaUtils::MakePlaybackParams(Context, Params, MediaSection);
 
-		SectionData->Setup(MediaPlayerStore, MediaSection, MediaPlayer, PlayerProxy, ProxyLayerIndex, ProxyTextureIndex, PlaybackParams);
+		SectionData->Setup(MediaPlayerStore, MediaSection, Params.MediaPlayer, MediaPlayerProxy, ProxyLayerIndex, ProxyTextureIndex, PlaybackParams);
 	}
 
 	if (!ensure(SectionData != nullptr))
@@ -798,6 +816,18 @@ void FMovieSceneMediaSectionTemplate::Initialize(const FMovieSceneEvaluationOper
 		return;
 	}
 
+	// If the player proxy became stale, we try to restore it.
+	// This may happen if the player proxy object is a spawnable.
+	if (SectionData->IsPlayerProxyStale())
+	{
+		// Note: if the proxy is not found, it is likely that it is not spawned yet.
+		// We keep it "stale" in the media section so that we keep trying to restore it until it is spawned again.
+		if (UObject* MediaPlayerProxy = UE::MediaSectionBaseMediaUtils::FindMediaPlayerProxyInterfaceObject(Operand, Player))
+		{
+			SectionData->UpdatePlayerProxy(MediaPlayerProxy);
+		}
+	}
+
 	UMediaPlayer* MediaPlayer = SectionData->GetMediaPlayer();
 
 	if (MediaPlayer == nullptr)
diff --git a/Engine/Plugins/Media/MediaCompositing/Source/MediaCompositing/Public/MovieSceneMediaData.h b/Engine/Plugins/Media/MediaCompositing/Source/MediaCompositing/Public/MovieSceneMediaData.h
--- a/Engine/Plugins/Media/MediaCompositing/Source/MediaCompositing/Public/MovieSceneMediaData.h
+++ b/Engine/Plugins/Media/MediaCompositing/Source/MediaCompositing/Public/MovieSceneMediaData.h
@@ -62,6 +62,11 @@ public:
 	 */
 	UObject* GetPlayerProxy() { return PlayerProxy.Get(); }
 
+	/**
+	 * Returns true if the player proxy object has become stale.
+	 */
+	bool IsPlayerProxyStale() { return PlayerProxy.IsStale(); }
+
 	/**
 	 * Get the layer index we are using (when using a proxy).
 	 */
@@ -85,6 +90,9 @@ public:
 	/** Set up this persistent data object. */
 	UE_API void Setup(const TSharedPtr<FMovieSceneMediaPlayerStore>& InMediaPlayerStore, const UObject* InSection, UMediaPlayer* OverrideMediaPlayer, UObject* InPlayerProxy, int32 InProxyLayerIndex, int32 InProxyTextureIndex, const FMovieSceneMediaPlaybackParams& InPlaybackParams);
 
+	/** Update the player proxy interface object. */
+	UE_API void UpdatePlayerProxy(UObject* InPlayerProxy);
+
 	/**
 	 * Called from FMovieSceneMediaSectionTemplate::Initialize.
 	 */
