[Media Track] UX/QOL Fix: Display "missing media texture" warning in the media section when the media texture is not specified (and should be).

diff --git a/Engine/Plugins/Media/MediaCompositing/Source/MediaCompositingEditor/Private/Sequencer/MediaThumbnailSection.cpp b/Engine/Plugins/Media/MediaCompositing/Source/MediaCompositingEditor/Private/Sequencer/MediaThumbnailSection.cpp
--- a/Engine/Plugins/Media/MediaCompositing/Source/MediaCompositingEditor/Private/Sequencer/MediaThumbnailSection.cpp
+++ b/Engine/Plugins/Media/MediaCompositing/Source/MediaCompositingEditor/Private/Sequencer/MediaThumbnailSection.cpp
@@ -173,7 +173,7 @@ float FMediaThumbnailSection::GetSectionHeight(const UE::Sequencer::FViewDensity
 
 	// Calculate minimum space for the section title, media info and optional performance warning.
 	const float PlayerInfoHeight = FAppStyle::GetFontStyle(MediaInfoFont).Size + 8.f;
-	const float NumInfoLines = bDrawSeekPerformanceWarning ? 2.0f : 1.0f;
+	const float NumInfoLines = CachedWarningString.IsEmpty() ? 1.0f : 2.0f;
 	const float MinimumHeight = SectionTitleHeight + PlayerInfoHeight * NumInfoLines + 8.0f;
 
 	// Base section height is either the thumbnail height or just the title.
@@ -279,6 +279,11 @@ void FMediaThumbnailSection::Tick(const FGeometry& AllottedGeometry, const FGeom
 {
 	UMovieSceneMediaSection* MediaSection = Cast<UMovieSceneMediaSection>(Section);
 
+	const bool bPreviousPerformanceWarning = bDrawSeekPerformanceWarning;
+	const bool bPreviousTextureWarning = bDrawMissingMediaTextureWarning;
+	
+	bDrawMissingMediaTextureWarning = false;
+	
 	if (MediaSection != nullptr)
 	{
 		if (GetDefault<UMovieSceneUserThumbnailSettings>()->bDrawSingleThumbnails)
@@ -289,13 +294,39 @@ void FMediaThumbnailSection::Tick(const FGeometry& AllottedGeometry, const FGeom
 		{
 			ThumbnailCache.SetSingleReferenceFrame(TOptional<double>());
 		}
+
+		// Logic for updating the "missing media texture" warning message.
+		if (!MediaSection->bUseExternalMediaPlayer				// Using external player has its own texture. (we can't check anyway)
+			&& !MediaSection->GetMediaSourceProxy().IsValid()	// Using a media source proxy will have its own texture.
+			&& MediaSection->MediaTexture == nullptr)			// Section texture is unspecified.
+		{
+			bDrawMissingMediaTextureWarning = true;
+		}
 	}
 
 	if (const UMediaPlayer* MediaPlayer = GetTemplateMediaPlayer())
 	{
 		UpdateCachedMediaInfo(MediaPlayer);
 	}
-	
+
+	if (bPreviousPerformanceWarning != bDrawSeekPerformanceWarning
+		|| bPreviousTextureWarning != bDrawMissingMediaTextureWarning
+		|| CachedWarningString.IsEmpty())
+	{
+		TArray<const FText*, TInlineAllocator<2>> Warnings;
+		if (bDrawSeekPerformanceWarning)
+		{
+			static const FText SeekPerformanceWarning = LOCTEXT("SeekPerformance", "Slow Seek Performance (GOP codec)");
+			Warnings.Add(&SeekPerformanceWarning);
+		}
+		if (bDrawMissingMediaTextureWarning)
+		{
+			static const FText MissingTextureWarning = LOCTEXT("MissingMediaTexture", "Missing Media Texture");
+			Warnings.Add(&MissingTextureWarning);
+		}
+		CachedWarningString = FString::JoinBy(Warnings, TEXT(" - "), [](const FText *InText) {return InText->ToString();});
+	}
+
 	FThumbnailSection::Tick(AllottedGeometry, ClippedGeometry, InCurrentTime, InDeltaTime);
 }
 
@@ -536,11 +567,10 @@ void FMediaThumbnailSection::DrawMediaInfo(FSequencerSectionPainter& InPainter,
 		float RemainingSize = AvailableSectionHeight - (TextYOffsetFromTop + TextSize.Y);
 
 		// Avoid having both messages overlapping.
-		if (bDrawSeekPerformanceWarning && RemainingSize > FontHeight)
+		if (!CachedWarningString.IsEmpty() && RemainingSize > FontHeight)
 		{
 			TGuardValue PushColor(TextPaintHelper.TextColor, FColor::Yellow);
-			FString SeekPerformanceWarning = LOCTEXT("SeekPerformance", "Slow Seek Performance (GOP codec)").ToString();
-			TextPaintHelper.PaintString(InPainter, SeekPerformanceWarning, 0.0f, EOffsetOrigin::FromBottom);
+			TextPaintHelper.PaintString(InPainter, CachedWarningString, 0.0f, EOffsetOrigin::FromBottom);
 		}
 	}
 }
diff --git a/Engine/Plugins/Media/MediaCompositing/Source/MediaCompositingEditor/Public/Sequencer/MediaThumbnailSection.h b/Engine/Plugins/Media/MediaCompositing/Source/MediaCompositingEditor/Public/Sequencer/MediaThumbnailSection.h
--- a/Engine/Plugins/Media/MediaCompositing/Source/MediaCompositingEditor/Public/Sequencer/MediaThumbnailSection.h
+++ b/Engine/Plugins/Media/MediaCompositing/Source/MediaCompositingEditor/Public/Sequencer/MediaThumbnailSection.h
@@ -150,6 +150,12 @@ private:
 	/** Show the seek performance warning. */
 	bool bDrawSeekPerformanceWarning = false;
 
+	/** Show the "missing media texture" warning. */
+	bool bDrawMissingMediaTextureWarning = false;
+
+	/** Cached Warning string to be displayed in the next paint. */
+	FString CachedWarningString;
+
 	/** True if the sequencer is currently scrubbing. */
 	bool bIsSequencerScrubbing = false;
 };
