[Media Track] Fixing media track media texture rendering priority when overlapping sections use the same media texture.

This method associates a priority to each player (kept in a common map) which is used to resolve which player will get the media texture in case of media section evaluation overlap.

diff --git a/Engine/Plugins/Media/MediaCompositing/Source/MediaCompositing/Private/MovieScene/MovieSceneMediaData.cpp b/Engine/Plugins/Media/MediaCompositing/Source/MediaCompositing/Private/MovieScene/MovieSceneMediaData.cpp
--- a/Engine/Plugins/Media/MediaCompositing/Source/MediaCompositing/Private/MovieScene/MovieSceneMediaData.cpp
+++ b/Engine/Plugins/Media/MediaCompositing/Source/MediaCompositing/Private/MovieScene/MovieSceneMediaData.cpp
@@ -85,6 +85,12 @@ void FMovieSceneMediaData::Setup(const TSharedPtr<FMovieSceneMediaPlayerStore>&
 		}
 	}
 
+	if (InMediaPlayerStore)
+	{
+		// The proxy layer index is the "row index" of the media section, which can be used to determine priority.
+		InMediaPlayerStore->SetMediaPlayerPriority(MediaPlayer, InProxyLayerIndex);
+	}
+	
 	MediaPlayer->PlayOnOpen = false;
 	MediaPlayer->OnMediaEvent().AddRaw(this, &FMovieSceneMediaData::HandleMediaPlayerEvent);
 	MediaPlayer->AddToRoot();
diff --git a/Engine/Plugins/Media/MediaCompositing/Source/MediaCompositing/Private/MovieScene/MovieSceneMediaPlayerStore.cpp b/Engine/Plugins/Media/MediaCompositing/Source/MediaCompositing/Private/MovieScene/MovieSceneMediaPlayerStore.cpp
--- a/Engine/Plugins/Media/MediaCompositing/Source/MediaCompositing/Private/MovieScene/MovieSceneMediaPlayerStore.cpp
+++ b/Engine/Plugins/Media/MediaCompositing/Source/MediaCompositing/Private/MovieScene/MovieSceneMediaPlayerStore.cpp
@@ -47,6 +47,32 @@ UMediaPlayer* FMovieSceneMediaPlayerStore::TryAcquireMediaPlayer(const FObjectKe
 	return nullptr;
 }
 
+void FMovieSceneMediaPlayerStore::SetMediaPlayerPriority(const UMediaPlayer* InMediaPlayer, int32 InPriority)
+{
+	// Prune stale keys to avoid unbounded growth.
+	for (TMap<FObjectKey, int32>::TIterator It(MediaPlayerPriorities); It; ++It)
+	{
+		if (!It->Key.ResolveObjectPtr())
+		{
+			It.RemoveCurrent();
+		}
+	}
+
+	MediaPlayerPriorities.Add(InMediaPlayer, InPriority);
+}
+
+TOptional<int32> FMovieSceneMediaPlayerStore::GetMediaPlayerPriority(const UMediaPlayer* InMediaPlayer) const
+{
+	if (InMediaPlayer)
+	{
+		if (const int32* Priority = MediaPlayerPriorities.Find(InMediaPlayer))
+		{
+			return *Priority;
+		}
+	}
+	return TOptional<int32>();
+}
+
 void FMovieSceneMediaPlayerStore::CloseMediaPlayer(UMediaPlayer* InMediaPlayer, bool bInCleanUpBeforeDestroy)
 {
 	if (IsValid(InMediaPlayer))
diff --git a/Engine/Plugins/Media/MediaCompositing/Source/MediaCompositing/Private/MovieScene/MovieSceneMediaPlayerStore.h b/Engine/Plugins/Media/MediaCompositing/Source/MediaCompositing/Private/MovieScene/MovieSceneMediaPlayerStore.h
--- a/Engine/Plugins/Media/MediaCompositing/Source/MediaCompositing/Private/MovieScene/MovieSceneMediaPlayerStore.h
+++ b/Engine/Plugins/Media/MediaCompositing/Source/MediaCompositing/Private/MovieScene/MovieSceneMediaPlayerStore.h
@@ -40,6 +40,25 @@ public:
 	 */
 	UMediaPlayer* TryAcquireMediaPlayer(const FObjectKey& InPersistentObjectKey);
 
+	/**
+	 * Associate a priority to media players used to play a media section.
+	 * This priority is used to resolve which player has access to a media texture in
+	 * case of overlap of multiple media sections with the same media texture.
+	 * 
+	 * @param InMediaPlayer Media player
+	 * @param InPriority Priority
+	 */
+	void SetMediaPlayerPriority(const UMediaPlayer* InMediaPlayer, int32 InPriority);
+
+	/**
+	 * Retrieve the associated priority for the media player (if it has one).
+	 * Only players managed by media section data will have a priority associated to them.
+	 * 
+	 * @param InMediaPlayer Media Player 
+	 * @return Associated priority, if found, for this player.
+	 */
+	TOptional<int32> GetMediaPlayerPriority(const UMediaPlayer* InMediaPlayer) const;
+	
 	/**
 	 * @brief Close the given media player and remove from root.
 	 * @param InMediaPlayer Input media player, validity is checked.
@@ -48,7 +67,7 @@ public:
 	static void CloseMediaPlayer(UMediaPlayer* InMediaPlayer, bool bInCleanUpBeforeDestroy = true);
 
 private:
-	/** Delegate handler for OneEndFrame. This is used to hook up CloseRemainingPlayers() somewhere after the end of the sequence evaluation. */
+	/** Delegate handler for OnEndFrame. This is used to hook up CloseRemainingPlayers() somewhere after the end of the sequence evaluation. */
 	void OnEndFrame();
 
 	/** Called after the evaluation is completed to close all non-acquired players. */
@@ -57,6 +76,9 @@ private:
 	/** Map of media players. */
 	TMap<FObjectKey, UMediaPlayer*> MediaPlayers;
 
+	/** Map of media player priorities used to resolve media texture access. */
+	TMap<FObjectKey, int32> MediaPlayerPriorities;
+
 	FDelegateHandle OnEndFrameDelegateHandle;
 };
 
diff --git a/Engine/Plugins/Media/MediaCompositing/Source/MediaCompositing/Private/MovieScene/MovieSceneMediaTemplate.cpp b/Engine/Plugins/Media/MediaCompositing/Source/MediaCompositing/Private/MovieScene/MovieSceneMediaTemplate.cpp
--- a/Engine/Plugins/Media/MediaCompositing/Source/MediaCompositing/Private/MovieScene/MovieSceneMediaTemplate.cpp
+++ b/Engine/Plugins/Media/MediaCompositing/Source/MediaCompositing/Private/MovieScene/MovieSceneMediaTemplate.cpp
@@ -857,16 +857,45 @@ void FMovieSceneMediaSectionTemplate::Initialize(const FMovieSceneEvaluationOper
 			Params.MediaSoundComponent->SetMediaPlayer(nullptr);
 		}
 	}
+	
+	TSharedPtr<FMediaTextureSampleQueue> PreRollSampleQueue;
+
+	if (IsEvaluating)
+	{
+		// Make sure to transfer the sample queue used for pre-roll (on the first evaluation), even if it
+		// isn't used by a media texture, otherwise it will block the player by not consuming the samples.
+		PreRollSampleQueue = SectionData->TransferSampleQueue();
+	}
 
 	if (Params.MediaTexture != nullptr)
 	{
 		if (IsEvaluating)
 		{
-			#if MOVIESCENEMEDIATEMPLATE_TRACE_EVALUATION
-				GLog->Logf(ELogVerbosity::Log, TEXT("Setting media player %p on media texture %p"), MediaPlayer, Params.MediaTexture.Get());
-			#endif
+			// Don't steal the media texture from higher priority players.
+			bool bCanUsePlayer = true;
+
+			if (const TSharedPtr<FMovieSceneMediaPlayerStore> MediaPlayerStore = FMovieSceneMediaPlayerStoreContainer::GetOrAdd(PersistentData).GetMediaPlayerStore())
+			{
+				// All players managed by a media data section will have an associated priority (derived from the media section's row index)
+				// that we can use to resolve which player should be used for that media texture. We compare the current player, i.e.
+				// the one currently assigned to the texture, with the new one from this section.
+				// Currently, since the priority is the row index, we want the lowest row index (i.e. top most in the sequencer)
+				// to be used. This corresponds to the "lower" priority value winning.
+				TOptional<int32> CurrentPriority = MediaPlayerStore->GetMediaPlayerPriority(Params.MediaTexture->GetMediaPlayer());
+				TOptional<int32> NewPriority = MediaPlayerStore->GetMediaPlayerPriority(MediaPlayer);
+				if (CurrentPriority.IsSet() && NewPriority.IsSet() && NewPriority.GetValue() > CurrentPriority.GetValue())
+				{
+					bCanUsePlayer = false;
+				}
+			}
 
-			Params.MediaTexture->SetMediaPlayer(MediaPlayer, SectionData->TransferSampleQueue());
+			if (bCanUsePlayer)
+			{
+#if MOVIESCENEMEDIATEMPLATE_TRACE_EVALUATION
+				GLog->Logf(ELogVerbosity::Log, TEXT("Setting media player %p on media texture %p"), MediaPlayer, Params.MediaTexture.Get());
+#endif
+				Params.MediaTexture->SetMediaPlayer(MediaPlayer, PreRollSampleQueue);
+			}
 		}
 		else if (Params.MediaTexture->GetMediaPlayer() == MediaPlayer)
 		{
@@ -877,15 +906,6 @@ void FMovieSceneMediaSectionTemplate::Initialize(const FMovieSceneEvaluationOper
 			Params.MediaTexture->SetMediaPlayer(nullptr);
 		}
 	}
-	else
-	{
-		// Make sure to discard the sample queue used for pre-roll if it isn't transferred to a media texture
-		// otherwise it will block the player by not consuming the samples.
-		if (IsEvaluating)
-		{
-			SectionData->TransferSampleQueue();
-		}
-	}
 
 	if (!UE::MediaSectionBaseMediaUtils::IsPlayerClosed(*MediaPlayer) && MediaPlayer->IsLooping() != Params.bLooping)
 	{
