[Media Track] Fix "missing media texture" message mistakenly appearing on sections under a media plate binding.

It was checking for a media source proxy, which is incorrect because a media section under a media plate may not have media source binding, but still be under the media plate's binding.
What we want to check is if the media secton has a media player proxy (there is already a cached flag for that in the section).

diff --git a/Engine/Plugins/Media/MediaCompositing/Source/MediaCompositingEditor/Private/Sequencer/MediaThumbnailSection.cpp b/Engine/Plugins/Media/MediaCompositing/Source/MediaCompositingEditor/Private/Sequencer/MediaThumbnailSection.cpp
--- a/Engine/Plugins/Media/MediaCompositing/Source/MediaCompositingEditor/Private/Sequencer/MediaThumbnailSection.cpp
+++ b/Engine/Plugins/Media/MediaCompositing/Source/MediaCompositingEditor/Private/Sequencer/MediaThumbnailSection.cpp
@@ -297,7 +297,7 @@ void FMediaThumbnailSection::Tick(const FGeometry& AllottedGeometry, const FGeom
 
 		// Logic for updating the "missing media texture" warning message.
 		if (!MediaSection->bUseExternalMediaPlayer				// Using external player has its own texture. (we can't check anyway)
-			&& !MediaSection->GetMediaSourceProxy().IsValid()	// Using a media source proxy will have its own texture.
+			&& !MediaSection->bHasMediaPlayerProxy				// Using a media player proxy will have its own texture.
 			&& MediaSection->MediaTexture == nullptr)			// Section texture is unspecified.
 		{
 			bDrawMissingMediaTextureWarning = true;
