[Media Track] Fixing inconsistent behavior - take 2.

Adding temporal resolution when sections are on the same row.

diff --git a/Engine/Plugins/Media/MediaCompositing/Source/MediaCompositing/Private/MovieScene/MovieSceneMediaData.cpp b/Engine/Plugins/Media/MediaCompositing/Source/MediaCompositing/Private/MovieScene/MovieSceneMediaData.cpp
--- a/Engine/Plugins/Media/MediaCompositing/Source/MediaCompositing/Private/MovieScene/MovieSceneMediaData.cpp
+++ b/Engine/Plugins/Media/MediaCompositing/Source/MediaCompositing/Private/MovieScene/MovieSceneMediaData.cpp
@@ -85,12 +85,6 @@ void FMovieSceneMediaData::Setup(const TSharedPtr<FMovieSceneMediaPlayerStore>&
 		}
 	}
 
-	if (InMediaPlayerStore)
-	{
-		// The proxy layer index is the "row index" of the media section, which can be used to determine priority.
-		InMediaPlayerStore->SetMediaPlayerPriority(MediaPlayer, InProxyLayerIndex);
-	}
-	
 	MediaPlayer->PlayOnOpen = false;
 	MediaPlayer->OnMediaEvent().AddRaw(this, &FMovieSceneMediaData::HandleMediaPlayerEvent);
 	MediaPlayer->AddToRoot();
diff --git a/Engine/Plugins/Media/MediaCompositing/Source/MediaCompositing/Private/MovieScene/MovieSceneMediaPlayerStore.cpp b/Engine/Plugins/Media/MediaCompositing/Source/MediaCompositing/Private/MovieScene/MovieSceneMediaPlayerStore.cpp
--- a/Engine/Plugins/Media/MediaCompositing/Source/MediaCompositing/Private/MovieScene/MovieSceneMediaPlayerStore.cpp
+++ b/Engine/Plugins/Media/MediaCompositing/Source/MediaCompositing/Private/MovieScene/MovieSceneMediaPlayerStore.cpp
@@ -47,10 +47,10 @@ UMediaPlayer* FMovieSceneMediaPlayerStore::TryAcquireMediaPlayer(const FObjectKe
 	return nullptr;
 }
 
-void FMovieSceneMediaPlayerStore::SetMediaPlayerPriority(const UMediaPlayer* InMediaPlayer, int32 InPriority)
+void FMovieSceneMediaPlayerStore::SetMediaPlayerPriority(const UMediaPlayer* InMediaPlayer, const FMovieSceneMediaPlayerPriority& InPriority)
 {
 	// Prune stale keys to avoid unbounded growth.
-	for (TMap<FObjectKey, int32>::TIterator It(MediaPlayerPriorities); It; ++It)
+	for (TMap<FObjectKey, FMovieSceneMediaPlayerPriority>::TIterator It(MediaPlayerPriorities); It; ++It)
 	{
 		if (!It->Key.ResolveObjectPtr())
 		{
@@ -61,16 +61,16 @@ void FMovieSceneMediaPlayerStore::SetMediaPlayerPriority(const UMediaPlayer* InM
 	MediaPlayerPriorities.Add(InMediaPlayer, InPriority);
 }
 
-TOptional<int32> FMovieSceneMediaPlayerStore::GetMediaPlayerPriority(const UMediaPlayer* InMediaPlayer) const
+TOptional<FMovieSceneMediaPlayerPriority> FMovieSceneMediaPlayerStore::GetMediaPlayerPriority(const UMediaPlayer* InMediaPlayer) const
 {
 	if (InMediaPlayer)
 	{
-		if (const int32* Priority = MediaPlayerPriorities.Find(InMediaPlayer))
+		if (const FMovieSceneMediaPlayerPriority* Priority = MediaPlayerPriorities.Find(InMediaPlayer))
 		{
 			return *Priority;
 		}
 	}
-	return TOptional<int32>();
+	return TOptional<FMovieSceneMediaPlayerPriority>();
 }
 
 void FMovieSceneMediaPlayerStore::CloseMediaPlayer(UMediaPlayer* InMediaPlayer, bool bInCleanUpBeforeDestroy)
@@ -96,6 +96,9 @@ void FMovieSceneMediaPlayerStore::CloseRemainingPlayers()
 	for (const TPair<FObjectKey, UMediaPlayer*>& MediaPlayerEntry : MediaPlayers)
 	{
 		CloseMediaPlayer(MediaPlayerEntry.Value);
+		
+		// Clean up associated auxiliary entry for this player.
+		MediaPlayerPriorities.Remove(MediaPlayerEntry.Value);
 	}
 	MediaPlayers.Reset();
 }
diff --git a/Engine/Plugins/Media/MediaCompositing/Source/MediaCompositing/Private/MovieScene/MovieSceneMediaPlayerStore.h b/Engine/Plugins/Media/MediaCompositing/Source/MediaCompositing/Private/MovieScene/MovieSceneMediaPlayerStore.h
--- a/Engine/Plugins/Media/MediaCompositing/Source/MediaCompositing/Private/MovieScene/MovieSceneMediaPlayerStore.h
+++ b/Engine/Plugins/Media/MediaCompositing/Source/MediaCompositing/Private/MovieScene/MovieSceneMediaPlayerStore.h
@@ -7,6 +7,42 @@
 
 class UMediaPlayer;
 
+/**
+ * Defines a priority for media players used in media sections.
+ * The priority is derived from the section's information and used for comparing
+ * media players when resolving evaluation conflicts.
+ */
+struct FMovieSceneMediaPlayerPriority
+{
+	/** Primary sort priority (hierarchical) */
+	int32 Priority = 0;
+	/** Secondary sort temporal priority */
+	FFrameNumber SectionStartFrame = FFrameNumber(0);
+
+	bool operator==(const FMovieSceneMediaPlayerPriority& InOther) const
+	{
+		return Priority == InOther.Priority && SectionStartFrame == InOther.SectionStartFrame;
+	}
+
+	bool operator >(const FMovieSceneMediaPlayerPriority& InOther) const
+	{
+		if (Priority != InOther.Priority)
+		{
+			return Priority > InOther.Priority;
+		}
+		return SectionStartFrame > InOther.SectionStartFrame;
+	}
+
+	bool operator <(const FMovieSceneMediaPlayerPriority& InOther) const
+	{
+		if (Priority != InOther.Priority)
+		{
+			return Priority < InOther.Priority;
+		}
+		return SectionStartFrame < InOther.SectionStartFrame;
+	}
+};
+
 /**
  * Keep a temporary store of media players so they can be reused after recompiling a media section.
  * During the track evaluation, the expired templates will be torn down first, then the new templates
@@ -48,7 +84,7 @@ public:
 	 * @param InMediaPlayer Media player
 	 * @param InPriority Priority
 	 */
-	void SetMediaPlayerPriority(const UMediaPlayer* InMediaPlayer, int32 InPriority);
+	void SetMediaPlayerPriority(const UMediaPlayer* InMediaPlayer, const FMovieSceneMediaPlayerPriority& InPriority);
 
 	/**
 	 * Retrieve the associated priority for the media player (if it has one).
@@ -57,7 +93,7 @@ public:
 	 * @param InMediaPlayer Media Player 
 	 * @return Associated priority, if found, for this player.
 	 */
-	TOptional<int32> GetMediaPlayerPriority(const UMediaPlayer* InMediaPlayer) const;
+	TOptional<FMovieSceneMediaPlayerPriority> GetMediaPlayerPriority(const UMediaPlayer* InMediaPlayer) const;
 	
 	/**
 	 * @brief Close the given media player and remove from root.
@@ -77,7 +113,7 @@ private:
 	TMap<FObjectKey, UMediaPlayer*> MediaPlayers;
 
 	/** Map of media player priorities used to resolve media texture access. */
-	TMap<FObjectKey, int32> MediaPlayerPriorities;
+	TMap<FObjectKey, FMovieSceneMediaPlayerPriority> MediaPlayerPriorities;
 
 	FDelegateHandle OnEndFrameDelegateHandle;
 };
diff --git a/Engine/Plugins/Media/MediaCompositing/Source/MediaCompositing/Private/MovieScene/MovieSceneMediaTemplate.cpp b/Engine/Plugins/Media/MediaCompositing/Source/MediaCompositing/Private/MovieScene/MovieSceneMediaTemplate.cpp
--- a/Engine/Plugins/Media/MediaCompositing/Source/MediaCompositing/Private/MovieScene/MovieSceneMediaTemplate.cpp
+++ b/Engine/Plugins/Media/MediaCompositing/Source/MediaCompositing/Private/MovieScene/MovieSceneMediaTemplate.cpp
@@ -809,6 +809,18 @@ void FMovieSceneMediaSectionTemplate::Initialize(const FMovieSceneEvaluationOper
 		FMovieSceneMediaPlaybackParams PlaybackParams = UE::MediaSectionBaseMediaUtils::MakePlaybackParams(Context, Params, MediaSection);
 
 		SectionData->Setup(MediaPlayerStore, MediaSection, Params.MediaPlayer, MediaPlayerProxy, ProxyLayerIndex, ProxyTextureIndex, PlaybackParams);
+
+		// Assign the media player priority to resolve overlapping access to a common media texture.
+		if (MediaPlayerStore && SectionData->GetMediaPlayer())
+		{
+			FMovieSceneMediaPlayerPriority PlayerPriority;
+			// The proxy layer index is the "row index" of the media section, which can be used to determine priority.
+			PlayerPriority.Priority = ProxyLayerIndex;
+			// For the "top-left" priority evaluation, we need the section start (leftmost boundary).
+			PlayerPriority.SectionStartFrame = Params.SectionStartFrame;
+			
+			MediaPlayerStore->SetMediaPlayerPriority(SectionData->GetMediaPlayer(), PlayerPriority);
+		}
 	}
 
 	if (!ensure(SectionData != nullptr))
@@ -876,13 +888,12 @@ void FMovieSceneMediaSectionTemplate::Initialize(const FMovieSceneEvaluationOper
 
 			if (const TSharedPtr<FMovieSceneMediaPlayerStore> MediaPlayerStore = FMovieSceneMediaPlayerStoreContainer::GetOrAdd(PersistentData).GetMediaPlayerStore())
 			{
-				// All players managed by a media data section will have an associated priority (derived from the media section's row index)
-				// that we can use to resolve which player should be used for that media texture. We compare the current player, i.e.
-				// the one currently assigned to the texture, with the new one from this section.
-				// Currently, since the priority is the row index, we want the lowest row index (i.e. top most in the sequencer)
-				// to be used. This corresponds to the "lower" priority value winning.
-				TOptional<int32> CurrentPriority = MediaPlayerStore->GetMediaPlayerPriority(Params.MediaTexture->GetMediaPlayer());
-				TOptional<int32> NewPriority = MediaPlayerStore->GetMediaPlayerPriority(MediaPlayer);
+				// Resolving concurrent media texture access using the associated media player priorities.
+				// We compare the current player, i.e. the one currently assigned to the media texture, with the new one from this section.
+				TOptional<FMovieSceneMediaPlayerPriority> CurrentPriority = MediaPlayerStore->GetMediaPlayerPriority(Params.MediaTexture->GetMediaPlayer());
+				TOptional<FMovieSceneMediaPlayerPriority> NewPriority = MediaPlayerStore->GetMediaPlayerPriority(MediaPlayer);
+
+				// Current rule is that the top-left most player gets the texture.
 				if (CurrentPriority.IsSet() && NewPriority.IsSet() && NewPriority.GetValue() > CurrentPriority.GetValue())
 				{
 					bCanUsePlayer = false;
