[MediaCompositing] Frame Accuracy Fix

The frame time is calculated with the display frame rate instead of the tick framerate. The result is that it will no longer move close to frame boundaries and produce inconsistent frame selection. Frame selection logic is sensibly different from player to player, especially around frame boundaries. Staying away from them helps selecting the expected frames regardless of where "in the frame" the sequencer last did the evaluation.

diff --git a/Engine/Plugins/Media/MediaCompositing/Source/MediaCompositing/Private/MovieScene/MovieSceneMediaTemplate.cpp b/Engine/Plugins/Media/MediaCompositing/Source/MediaCompositing/Private/MovieScene/MovieSceneMediaTemplate.cpp
--- a/Engine/Plugins/Media/MediaCompositing/Source/MediaCompositing/Private/MovieScene/MovieSceneMediaTemplate.cpp
+++ b/Engine/Plugins/Media/MediaCompositing/Source/MediaCompositing/Private/MovieScene/MovieSceneMediaTemplate.cpp
@@ -115,18 +115,101 @@ namespace UE::MediaSectionBaseMediaUtils
 	}
 
 	/**
-	 * Utility to convert a frame number to a timespan at the given frame rate. 
+	 * Utility function to obtain the "best" display frame rate we can have given the context.
+	 * @param InContext Evaluation context, which should contain a "tick" rate we can fall back to.
+	 * @param InMediaSection Media section should give us the parent movie scene to obtain actual display rate.
+	 * @return movie scene display rate if available, sequencer tick rate otherwise.
+	 */
+	FFrameRate GetDisplayFrameRate(const FMovieSceneContext& InContext, const UMovieSceneMediaSection* InMediaSection)
+	{
+		if (InMediaSection)
+		{
+			if (UMovieScene* MovieScene = InMediaSection->GetTypedOuter<UMovieScene>())
+			{
+				return MovieScene->GetDisplayRate();
+			}
+		}
+		// As a fallback, we use the tick frame rate, which may also be the display rate when sequencer is running in locked mode.
+		return InContext.GetFrameRate();
+	}
+
+	/**
+	 * Utility to convert a frame number to a timespan at the given frame rate.
+	 * The timespan is either the first or second tick of the frame.
+	 * This can be used to specify the player's block on range accurately.
 	 */
 	FTimespan FrameNumberToTimespan(const FFrameNumber& InFrameNumber, const FFrameRate& InFrameRate)
 	{
-		// Caution: A larger denominator limits the frame number range. 
-		// Example: for 59.94, 29.97 or 23.975, because denominator is 1001, the maximum frame number is going to be 1073741823 (which is ~207 days at 59.94).
-		const int64 DenominatorTicks = InFrameRate.Denominator * ETimespan::TicksPerSecond;
-		// Using Ceil to ensure the tick value is part of the desired frame, and not the frame before
-		// regardless of the internal media player's implementation of "time to frame".
-		return FTimespan(FMath::DivideAndRoundUp(int64(InFrameNumber.Value * DenominatorTicks), int64(InFrameRate.Numerator)));
+		if (InFrameRate.Numerator != 0)
+		{
+			// Caution: A larger denominator limits the frame number range. 
+			// Example: for 59.94, 29.97 or 23.975, because denominator is 1001, the maximum frame number is going to be 1073741823 (which is ~207 days at 59.94).
+			const int64 DenominatorTicks = InFrameRate.Denominator * ETimespan::TicksPerSecond;
+			// Using Ceil to ensure the tick value is part of the desired frame, and not the frame before
+			// regardless of the internal media player's implementation of "time to frame".
+			return FTimespan(FMath::DivideAndRoundUp(int64(InFrameNumber.Value * DenominatorTicks), int64(InFrameRate.Numerator)));
+		}
+		return FTimespan();
 	}
-	
+
+	/**
+	 * Utility function to calculate the exact frame duration for the given frame, so that
+	 * FrameTick + FrameDurationTick = exactly the last tick of the frame. 
+	 * Calculating this way takes into account the different rounding given the frame rate.
+	 * Some frame rates, like 30 fps have an oscillating frame duration when rounded in ticks.
+	 * 
+	 * @param InFrameNumber Frame number we want to calculate the exact duration of.
+	 * @param InFrameRate Frame rate.
+	 * @return An exact frame duration in ticks.
+	 */
+	FTimespan GetFrameDurationTimespan(const FFrameNumber& InFrameNumber, const FFrameRate& InFrameRate)
+	{
+		const FTimespan FrameStartTime = FrameNumberToTimespan(InFrameNumber, InFrameRate);
+		const FTimespan NextFrameStartTime = FrameNumberToTimespan(InFrameNumber+ 1, InFrameRate);
+		return FTimespan(FMath::Max(NextFrameStartTime.GetTicks() - FrameStartTime.GetTicks() - 1, 1));
+	}
+
+	/**
+	 * Returns the frame duration rounded to nearest tick. So, it may have +- 1 tick of inaccuracy. 
+	 * Note: for Block on range specification, the frame accurate version, GetFrameDurationTimespan,
+	 * should be used instead.
+	 * @param InFrameRate Input frame rate 
+	 * @return rounded frame duration in ticks.
+	 */
+	FTimespan GetRoundedFrameDurationTimespan(const FFrameRate& InFrameRate)
+	{
+		return InFrameRate.Numerator != 0 ? FTimespan::FromSeconds(InFrameRate.AsInterval()) : FTimespan(1);
+	}
+
+	/**
+	 * Helper function to get best display frame duration given the context.
+	 * Same note applies as for GetRoundedFrameDurationTimespan.
+	 * 
+	 * @param InContext Evaluation context
+	 * @param InMediaSection Media section to get display rate from.
+	 * @return Rounded frame duration.
+	 */
+	FTimespan GetRoundedDisplayFrameDurationTimespan(const FMovieSceneContext& InContext, const UMovieSceneMediaSection* InMediaSection)
+	{
+		return GetRoundedFrameDurationTimespan(GetDisplayFrameRate(InContext, InMediaSection));
+	}
+
+	/**
+	 * Calculates the frame number aligned to the given display frame rate.
+	 * @param InTickFrameNumber Current sequencer tick frame number
+	 * @param InTickFrameRate Current sequencer tick framerate (24000, 60000, etc)
+	 * @param InDisplayFrameRate Current sequencer display framerate (24, 30, 60, etc)
+	 * @return Calculated sequencer display frame number.
+	 */
+	FFrameNumber GetDisplayFrame(const FFrameNumber& InTickFrameNumber, const FFrameRate& InTickFrameRate, const FFrameRate& InDisplayFrameRate)
+	{
+		if (InTickFrameRate.Numerator != 0 && InDisplayFrameRate.Denominator != 0)
+		{
+			return FFrameRate::TransformTime(FFrameTime(InTickFrameNumber), InTickFrameRate, InDisplayFrameRate).GetFrame();
+		}
+		return FFrameNumber(0);
+	}
+
 	/**
 	 * Utility to convert a sequencer's (absolute) frame number to a player's timespan.
 	 */
@@ -230,23 +313,13 @@ namespace UE::MediaSectionBaseMediaUtils
 		InMediaPlayer->SetBlockOnTimeRange(TRange<FTimespan>(RangeLowerBound, RangeUpperBound));
 	}
 
-	/** Calculate the timespan of a frame duration given the sequencer's framerate */
-	FTimespan GetFrameDuration(const FMovieSceneContext& InContext) 
-	{
-		const FFrameRate FrameRate = InContext.GetFrameRate();
-		// With zero-length frames (which can occur occasionally), we use the fixed frame time, matching previous behavior.
-		const double FrameDurationInSeconds = FMath::Max(FrameRate.AsSeconds(FFrameTime(1)), (InContext.GetRange().Size<FFrameTime>()) / FrameRate);
-		const int64 FrameDurationTicks = static_cast<int64>(FrameDurationInSeconds * ETimespan::TicksPerSecond);
-		return FTimespan(FrameDurationTicks);
-	}
-
 	/** Prepare the playback parameters for the execution token. */
-	FMovieSceneMediaPlaybackParams MakePlaybackParams(const FMovieSceneContext& InContext, const FMovieSceneMediaSectionParams& InParams, const UMovieSceneMediaSection* InMediaSection)
+	FMovieSceneMediaPlaybackParams MakePlaybackParams(const FMovieSceneContext& InContext, const FMovieSceneMediaSectionParams& InParams, const UMovieSceneMediaSection* InMediaSection, const FTimespan& InFrameDuration)
 	{
 		FMovieSceneMediaPlaybackParams Params;
 		Params.SectionTimeRange = CVarUpdatePlaybackRange.GetValueOnGameThread() ?
 			CalculateSectionPlaybackTimeRange(InParams, InMediaSection, InContext.GetFrameRate()) : TRange<FTimespan>::Empty(); 
-		Params.FrameDuration = GetFrameDuration(InContext);
+		Params.FrameDuration = InFrameDuration;
 		Params.bIsLooping = InParams.bLooping;
 		return Params;
 	}
@@ -419,12 +492,8 @@ private:
 };
 
 struct FMediaSectionPostRollExecutionToken
-	: FMediaSectionBaseExecutionToken
+	: IMovieSceneExecutionToken
 {
-	FMediaSectionPostRollExecutionToken(UMediaSource* InMediaSource, const FMediaSectionBaseMediaSourceOptions& InMediaSourceOptions, const FMovieSceneMediaPlaybackParams& InPlaybackParams, FMovieSceneObjectBindingID InMediaSourceProxy, int32 InMediaSourceProxyIndex)
-		: FMediaSectionBaseExecutionToken(InMediaSource, InMediaSourceOptions, InPlaybackParams, InMediaSourceProxy, InMediaSourceProxyIndex)
-	{ }
-
 	virtual void Execute(const FMovieSceneContext& Context, const FMovieSceneEvaluationOperand& Operand, FPersistentEvaluationData& PersistentData, IMovieScenePlayer& Player) override
 	{
 		FMovieSceneMediaData& SectionData = PersistentData.GetSectionData<FMovieSceneMediaData>();
@@ -709,31 +778,32 @@ void FMovieSceneMediaSectionTemplate::Evaluate(const FMovieSceneEvaluationOperan
 	bool bCanPlayerBeOpen = true;
 	MediaSection->ChannelCanPlayerBeOpen.Evaluate(Context.GetTime(), bCanPlayerBeOpen);
 
+	using namespace UE::MediaSectionBaseMediaUtils;
+	const FFrameRate DisplayRate = GetDisplayFrameRate(Context, MediaSection);
+	
 	if (Context.IsPreRoll() && bCanPlayerBeOpen)
 	{
-		const FFrameRate FrameRate = Context.GetFrameRate();
-		const FFrameNumber StartFrame = Context.HasPreRollEndTime() ? Context.GetPreRollEndFrame() - Params.SectionStartFrame + Params.StartFrameOffset : Params.StartFrameOffset;
-		const FTimespan StartTime = UE::MediaSectionBaseMediaUtils::FrameNumberToTimespan(StartFrame, FrameRate);
+		const FFrameNumber StartTickFrame = Context.HasPreRollEndTime() ? Context.GetPreRollEndFrame() - Params.SectionStartFrame + Params.StartFrameOffset : Params.StartFrameOffset;
+		const FFrameNumber StartDisplayFrame = GetDisplayFrame(StartTickFrame, Context.GetFrameRate(), DisplayRate); 
+		const FTimespan StartTime = FrameNumberToTimespan(StartDisplayFrame, DisplayRate);
+		const FTimespan DisplayFrameDuration = GetFrameDurationTimespan(StartDisplayFrame, DisplayRate);
 
 		FMediaSectionBaseMediaSourceOptions Options;
 		Options.CacheSettings = Params.CacheSettings;
 		Options.bSynchronousScrubbing = false; // not relevant during pre-roll
-		const FMovieSceneMediaPlaybackParams PlaybackParams = UE::MediaSectionBaseMediaUtils::MakePlaybackParams(Context, Params, MediaSection);
+		const FMovieSceneMediaPlaybackParams PlaybackParams = MakePlaybackParams(Context, Params, MediaSection, DisplayFrameDuration);
 		ExecutionTokens.Add(FMediaSectionPreRollExecutionToken(MediaSource, Options, PlaybackParams, Params.MediaSourceProxy, Params.MediaSourceProxyIndex, StartTime));
 	}
 	else if (Context.IsPostRoll() && bCanPlayerBeOpen)
 	{
-		FMediaSectionBaseMediaSourceOptions Options;
-		Options.CacheSettings = Params.CacheSettings;
-		Options.bSynchronousScrubbing = false; // not relevant during post-roll
-		const FMovieSceneMediaPlaybackParams PlaybackParams = UE::MediaSectionBaseMediaUtils::MakePlaybackParams(Context, Params, MediaSection);
-		ExecutionTokens.Add(FMediaSectionPostRollExecutionToken(MediaSource, Options, PlaybackParams, Params.MediaSourceProxy, Params.MediaSourceProxyIndex));
+		ExecutionTokens.Add(FMediaSectionPostRollExecutionToken());
 	}
 	else if (!Context.IsPostRoll() && (Context.GetTime().FrameNumber < Params.SectionEndFrame))
 	{
-		const FFrameRate FrameRate = Context.GetFrameRate();
-		const FFrameNumber FrameNumber = Context.GetTime().FrameNumber - Params.SectionStartFrame + Params.StartFrameOffset;
-		const FTimespan FrameTime = UE::MediaSectionBaseMediaUtils::FrameNumberToTimespan(FrameNumber, FrameRate);
+		const FFrameNumber TickFrameNumber = Context.GetTime().FrameNumber - Params.SectionStartFrame + Params.StartFrameOffset;
+		const FFrameNumber DisplayFrameNumber = GetDisplayFrame(TickFrameNumber, Context.GetFrameRate(), DisplayRate); 
+		const FTimespan FrameTime = FrameNumberToTimespan(DisplayFrameNumber, DisplayRate);
+		const FTimespan DisplayFrameDuration = GetFrameDurationTimespan(DisplayFrameNumber, DisplayRate);
 		
 		const float ProxyTextureBlend = MediaSection->EvaluateEasing(Context.GetTime());
 
@@ -745,7 +815,7 @@ void FMovieSceneMediaSectionTemplate::Evaluate(const FMovieSceneEvaluationOperan
 				FrameRate.Numerator,
 				FrameRate.Denominator,
 				FrameTime.GetTicks(),
-				UE::MediaSectionBaseMediaUtils::GetFrameDuration(Context).GetTicks()
+				DisplayFrameDuration.GetTicks()
 			);
 		#endif
 
@@ -754,7 +824,7 @@ void FMovieSceneMediaSectionTemplate::Evaluate(const FMovieSceneEvaluationOperan
 		FMediaSectionBaseMediaSourceOptions Options;
 		Options.CacheSettings = Params.CacheSettings;
 		Options.bSynchronousScrubbing = MediaTrack ? MediaTrack->bSynchronousScrubbing : false;
-		FMovieSceneMediaPlaybackParams PlaybackParams = UE::MediaSectionBaseMediaUtils::MakePlaybackParams(Context, Params, MediaSection);
+		FMovieSceneMediaPlaybackParams PlaybackParams = MakePlaybackParams(Context, Params, MediaSection, DisplayFrameDuration);
 
 		// Scrubbing or stepping outside the playback range is allowed by the sequencer,
 		// in which case we have to reset the player's playback range.
@@ -794,19 +864,22 @@ void FMovieSceneMediaSectionTemplate::Initialize(const FMovieSceneEvaluationOper
 			ProxyLayerIndex = MediaSection->GetRowIndex();
 		}
 
+		using namespace UE::MediaSectionBaseMediaUtils;
+
 		// If the media section doesn't specify an external media player, we try to find the
 		// media player proxy from the bound objects.
 		UObject* MediaPlayerProxy = nullptr;
 		if (Params.MediaPlayer == nullptr)
 		{
-			MediaPlayerProxy = UE::MediaSectionBaseMediaUtils::FindMediaPlayerProxyInterfaceObject(Operand, Player);
+			MediaPlayerProxy = FindMediaPlayerProxyInterfaceObject(Operand, Player);
 		}
 
 		// Add section data.
 		SectionData = &PersistentData.AddSectionData<FMovieSceneMediaData>();
 		const TSharedPtr<FMovieSceneMediaPlayerStore> MediaPlayerStore = FMovieSceneMediaPlayerStoreContainer::GetOrAdd(PersistentData).GetMediaPlayerStore();
 		// Prepare these values in the section data to have them set in the player when it is opened (as early as possible before the first seek).
-		FMovieSceneMediaPlaybackParams PlaybackParams = UE::MediaSectionBaseMediaUtils::MakePlaybackParams(Context, Params, MediaSection);
+		const FTimespan FrameDuration = GetRoundedDisplayFrameDurationTimespan(Context, MediaSection);
+		FMovieSceneMediaPlaybackParams PlaybackParams = MakePlaybackParams(Context, Params, MediaSection, FrameDuration);
 
 		SectionData->Setup(MediaPlayerStore, MediaSection, Params.MediaPlayer, MediaPlayerProxy, ProxyLayerIndex, ProxyTextureIndex, PlaybackParams);
 
