[MediaCompositing] Sequencer Media Track: Revisiting the frame alignment for frame accuracy.

We align the "Block On Range" (BOR) on the media's frame rate to ensure accuracy rather than aligning on the sequencer display rate. The sequencer display rate is not meant to be used at runtime.

Instead, the media section now has an alignment frame rate override available in case the media frame rate is incorrect (some formats are either variable frame rate or have incorrect meta data).

diff --git a/Engine/Plugins/Media/MediaCompositing/Source/MediaCompositing/Private/MovieScene/MovieSceneMediaSection.cpp b/Engine/Plugins/Media/MediaCompositing/Source/MediaCompositing/Private/MovieScene/MovieSceneMediaSection.cpp
--- a/Engine/Plugins/Media/MediaCompositing/Source/MediaCompositing/Private/MovieScene/MovieSceneMediaSection.cpp
+++ b/Engine/Plugins/Media/MediaCompositing/Source/MediaCompositing/Private/MovieScene/MovieSceneMediaSection.cpp
@@ -28,7 +28,10 @@ UMovieSceneMediaSection::UMovieSceneMediaSection(const FObjectInitializer& Objec
 	: Super(ObjectInitializer)
 	, MediaSourceProxyIndex(0)
 	, bLooping(true)
+	, bUseExternalMediaPlayer(false)
 	, TextureIndex(0)
+	, bManualFrameRateAlignment(false)
+	, FrameRateAlignment(24, 1)
 	, bHasMediaPlayerProxy(false)
 {
 #if WITH_EDITORONLY_DATA
diff --git a/Engine/Plugins/Media/MediaCompositing/Source/MediaCompositing/Private/MovieScene/MovieSceneMediaTemplate.cpp b/Engine/Plugins/Media/MediaCompositing/Source/MediaCompositing/Private/MovieScene/MovieSceneMediaTemplate.cpp
--- a/Engine/Plugins/Media/MediaCompositing/Source/MediaCompositing/Private/MovieScene/MovieSceneMediaTemplate.cpp
+++ b/Engine/Plugins/Media/MediaCompositing/Source/MediaCompositing/Private/MovieScene/MovieSceneMediaTemplate.cpp
@@ -13,8 +13,6 @@
 #include "MediaSourceOptions.h"
 #include "MediaTexture.h"
 #include "MovieScene.h"
-#include "UObject/Package.h"
-#include "UObject/GCObject.h"
 
 #include "MovieSceneMediaData.h"
 #include "MovieSceneMediaPlayerStore.h"
@@ -115,20 +113,40 @@ namespace UE::MediaSectionBaseMediaUtils
 	}
 
 	/**
-	 * Utility function to obtain the "best" display frame rate we can have given the context.
+	 * Utility function to obtain the frame rate used for alignment of the "block on range".
 	 * @param InContext Evaluation context, which should contain a "tick" rate we can fall back to.
 	 * @param InMediaSection Media section should give us the parent movie scene to obtain actual display rate.
+	 * @param InPersistentData Persistent evaluation data used to retrieve the media section data.
 	 * @return movie scene display rate if available, sequencer tick rate otherwise.
 	 */
-	FFrameRate GetDisplayFrameRate(const FMovieSceneContext& InContext, const UMovieSceneMediaSection* InMediaSection)
+	FFrameRate GetAlignedFrameRate(const FMovieSceneContext& InContext, const UMovieSceneMediaSection* InMediaSection, const FPersistentEvaluationData& InPersistentData)
 	{
-		if (InMediaSection)
+		// Manual frame alignment has been specified.
+		if (InMediaSection && InMediaSection->bManualFrameRateAlignment)
 		{
-			if (UMovieScene* MovieScene = InMediaSection->GetTypedOuter<UMovieScene>())
+			return InMediaSection->FrameRateAlignment;
+		}
+
+		// Try to get the media's frame rate.
+		if (FMovieSceneMediaData* SectionData = InPersistentData.FindSectionData<FMovieSceneMediaData>())
+		{
+			if (UMediaPlayer* MediaPlayer = SectionData->GetMediaPlayer())
 			{
-				return MovieScene->GetDisplayRate();
+				const int32 SelectedTrack = MediaPlayer->GetSelectedTrack(EMediaPlayerTrack::Video);
+				const int32 SelectedFormat = MediaPlayer->GetTrackFormat(EMediaPlayerTrack::Video, SelectedTrack);
+				const float Framerate = MediaPlayer->GetVideoTrackFrameRate(SelectedTrack, SelectedFormat);
+				if (Framerate > 0.0f)
+				{
+					// Fractional frame rates are considered "drop frame"
+					if (FMath::Fractional(Framerate) > UE_KINDA_SMALL_NUMBER)
+					{
+						return FFrameRate(FMath::CeilToInt(Framerate) * 1000, 1001);
+					}
+					return FFrameRate(Framerate, 1);
+				}
 			}
 		}
+
 		// As a fallback, we use the tick frame rate, which may also be the display rate when sequencer is running in locked mode.
 		return InContext.GetFrameRate();
 	}
@@ -162,7 +180,7 @@ namespace UE::MediaSectionBaseMediaUtils
 	 * @param InFrameRate Frame rate.
 	 * @return An exact frame duration in ticks.
 	 */
-	FTimespan GetFrameDurationTimespan(const FFrameNumber& InFrameNumber, const FFrameRate& InFrameRate)
+	FTimespan GetFrameDuration(const FFrameNumber& InFrameNumber, const FFrameRate& InFrameRate)
 	{
 		const FTimespan FrameStartTime = FrameNumberToTimespan(InFrameNumber, InFrameRate);
 		const FTimespan NextFrameStartTime = FrameNumberToTimespan(InFrameNumber+ 1, InFrameRate);
@@ -171,41 +189,42 @@ namespace UE::MediaSectionBaseMediaUtils
 
 	/**
 	 * Returns the frame duration rounded to nearest tick. So, it may have +- 1 tick of inaccuracy. 
-	 * Note: for Block on range specification, the frame accurate version, GetFrameDurationTimespan,
+	 * Note: for Block on range specification, the frame accurate version, GetFrameDuration,
 	 * should be used instead.
 	 * @param InFrameRate Input frame rate 
 	 * @return rounded frame duration in ticks.
 	 */
-	FTimespan GetRoundedFrameDurationTimespan(const FFrameRate& InFrameRate)
+	FTimespan GetRoundedFrameDuration(const FFrameRate& InFrameRate)
 	{
 		return InFrameRate.Numerator != 0 ? FTimespan::FromSeconds(InFrameRate.AsInterval()) : FTimespan(1);
 	}
 
 	/**
-	 * Helper function to get best display frame duration given the context.
-	 * Same note applies as for GetRoundedFrameDurationTimespan.
+	 * Helper function to get the aligned frame duration given the context.
+	 * Same note applies as for GetRoundedFrameDuration.
 	 * 
 	 * @param InContext Evaluation context
 	 * @param InMediaSection Media section to get display rate from.
+	 * @param InPersistentData Persistent evaluation data used to retrieve the media section data.
 	 * @return Rounded frame duration.
 	 */
-	FTimespan GetRoundedDisplayFrameDurationTimespan(const FMovieSceneContext& InContext, const UMovieSceneMediaSection* InMediaSection)
+	FTimespan GetRoundedAlignedFrameDuration(const FMovieSceneContext& InContext, const UMovieSceneMediaSection* InMediaSection, const FPersistentEvaluationData& InPersistentData)
 	{
-		return GetRoundedFrameDurationTimespan(GetDisplayFrameRate(InContext, InMediaSection));
+		return GetRoundedFrameDuration(GetAlignedFrameRate(InContext, InMediaSection, InPersistentData));
 	}
 
 	/**
-	 * Calculates the frame number aligned to the given display frame rate.
+	 * Calculates the frame number aligned to the given alignment frame rate.
 	 * @param InTickFrameNumber Current sequencer tick frame number
 	 * @param InTickFrameRate Current sequencer tick framerate (24000, 60000, etc)
-	 * @param InDisplayFrameRate Current sequencer display framerate (24, 30, 60, etc)
+	 * @param InAlignmentFrameRate Desired alignment frame rate for the media (24, 60, etc)
 	 * @return Calculated sequencer display frame number.
 	 */
-	FFrameNumber GetDisplayFrame(const FFrameNumber& InTickFrameNumber, const FFrameRate& InTickFrameRate, const FFrameRate& InDisplayFrameRate)
+	FFrameNumber GetAlignedFrame(const FFrameNumber& InTickFrameNumber, const FFrameRate& InTickFrameRate, const FFrameRate& InAlignmentFrameRate)
 	{
-		if (InTickFrameRate.Numerator != 0 && InDisplayFrameRate.Denominator != 0)
+		if (InTickFrameRate.Numerator != 0 && InAlignmentFrameRate.Denominator != 0)
 		{
-			return FFrameRate::TransformTime(FFrameTime(InTickFrameNumber), InTickFrameRate, InDisplayFrameRate).GetFrame();
+			return FFrameRate::TransformTime(FFrameTime(InTickFrameNumber), InTickFrameRate, InAlignmentFrameRate).GetFrame();
 		}
 		return FFrameNumber(0);
 	}
@@ -779,19 +798,19 @@ void FMovieSceneMediaSectionTemplate::Evaluate(const FMovieSceneEvaluationOperan
 	MediaSection->ChannelCanPlayerBeOpen.Evaluate(Context.GetTime(), bCanPlayerBeOpen);
 
 	using namespace UE::MediaSectionBaseMediaUtils;
-	const FFrameRate DisplayRate = GetDisplayFrameRate(Context, MediaSection);
+	const FFrameRate AlignedFrameRate = GetAlignedFrameRate(Context, MediaSection, PersistentData);
 	
 	if (Context.IsPreRoll() && bCanPlayerBeOpen)
 	{
 		const FFrameNumber StartTickFrame = Context.HasPreRollEndTime() ? Context.GetPreRollEndFrame() - Params.SectionStartFrame + Params.StartFrameOffset : Params.StartFrameOffset;
-		const FFrameNumber StartDisplayFrame = GetDisplayFrame(StartTickFrame, Context.GetFrameRate(), DisplayRate); 
-		const FTimespan StartTime = FrameNumberToTimespan(StartDisplayFrame, DisplayRate);
-		const FTimespan DisplayFrameDuration = GetFrameDurationTimespan(StartDisplayFrame, DisplayRate);
+		const FFrameNumber StartAlignedFrame = GetAlignedFrame(StartTickFrame, Context.GetFrameRate(), AlignedFrameRate); 
+		const FTimespan StartTime = FrameNumberToTimespan(StartAlignedFrame, AlignedFrameRate);
+		const FTimespan AlignedFrameDuration = GetFrameDuration(StartAlignedFrame, AlignedFrameRate);
 
 		FMediaSectionBaseMediaSourceOptions Options;
 		Options.CacheSettings = Params.CacheSettings;
 		Options.bSynchronousScrubbing = false; // not relevant during pre-roll
-		const FMovieSceneMediaPlaybackParams PlaybackParams = MakePlaybackParams(Context, Params, MediaSection, DisplayFrameDuration);
+		const FMovieSceneMediaPlaybackParams PlaybackParams = MakePlaybackParams(Context, Params, MediaSection, AlignedFrameDuration);
 		ExecutionTokens.Add(FMediaSectionPreRollExecutionToken(MediaSource, Options, PlaybackParams, Params.MediaSourceProxy, Params.MediaSourceProxyIndex, StartTime));
 	}
 	else if (Context.IsPostRoll() && bCanPlayerBeOpen)
@@ -801,9 +820,9 @@ void FMovieSceneMediaSectionTemplate::Evaluate(const FMovieSceneEvaluationOperan
 	else if (!Context.IsPostRoll() && (Context.GetTime().FrameNumber < Params.SectionEndFrame))
 	{
 		const FFrameNumber TickFrameNumber = Context.GetTime().FrameNumber - Params.SectionStartFrame + Params.StartFrameOffset;
-		const FFrameNumber DisplayFrameNumber = GetDisplayFrame(TickFrameNumber, Context.GetFrameRate(), DisplayRate); 
-		const FTimespan FrameTime = FrameNumberToTimespan(DisplayFrameNumber, DisplayRate);
-		const FTimespan DisplayFrameDuration = GetFrameDurationTimespan(DisplayFrameNumber, DisplayRate);
+		const FFrameNumber AlignedFrameNumber = GetAlignedFrame(TickFrameNumber, Context.GetFrameRate(), AlignedFrameRate); 
+		const FTimespan FrameTime = FrameNumberToTimespan(AlignedFrameNumber, AlignedFrameRate);
+		const FTimespan AlignedFrameDuration = GetFrameDuration(AlignedFrameNumber, AlignedFrameRate);
 		
 		const float ProxyTextureBlend = MediaSection->EvaluateEasing(Context.GetTime());
 
@@ -815,7 +834,7 @@ void FMovieSceneMediaSectionTemplate::Evaluate(const FMovieSceneEvaluationOperan
 				FrameRate.Numerator,
 				FrameRate.Denominator,
 				FrameTime.GetTicks(),
-				DisplayFrameDuration.GetTicks()
+				AlignedFrameDuration.GetTicks()
 			);
 		#endif
 
@@ -824,7 +843,7 @@ void FMovieSceneMediaSectionTemplate::Evaluate(const FMovieSceneEvaluationOperan
 		FMediaSectionBaseMediaSourceOptions Options;
 		Options.CacheSettings = Params.CacheSettings;
 		Options.bSynchronousScrubbing = MediaTrack ? MediaTrack->bSynchronousScrubbing : false;
-		FMovieSceneMediaPlaybackParams PlaybackParams = MakePlaybackParams(Context, Params, MediaSection, DisplayFrameDuration);
+		FMovieSceneMediaPlaybackParams PlaybackParams = MakePlaybackParams(Context, Params, MediaSection, AlignedFrameDuration);
 
 		// Scrubbing or stepping outside the playback range is allowed by the sequencer,
 		// in which case we have to reset the player's playback range.
@@ -878,7 +897,7 @@ void FMovieSceneMediaSectionTemplate::Initialize(const FMovieSceneEvaluationOper
 		SectionData = &PersistentData.AddSectionData<FMovieSceneMediaData>();
 		const TSharedPtr<FMovieSceneMediaPlayerStore> MediaPlayerStore = FMovieSceneMediaPlayerStoreContainer::GetOrAdd(PersistentData).GetMediaPlayerStore();
 		// Prepare these values in the section data to have them set in the player when it is opened (as early as possible before the first seek).
-		const FTimespan FrameDuration = GetRoundedDisplayFrameDurationTimespan(Context, MediaSection);
+		const FTimespan FrameDuration = GetRoundedAlignedFrameDuration(Context, MediaSection, PersistentData);
 		FMovieSceneMediaPlaybackParams PlaybackParams = MakePlaybackParams(Context, Params, MediaSection, FrameDuration);
 
 		SectionData->Setup(MediaPlayerStore, MediaSection, Params.MediaPlayer, MediaPlayerProxy, ProxyLayerIndex, ProxyTextureIndex, PlaybackParams);
diff --git a/Engine/Plugins/Media/MediaCompositing/Source/MediaCompositing/Public/MovieSceneMediaSection.h b/Engine/Plugins/Media/MediaCompositing/Source/MediaCompositing/Public/MovieSceneMediaSection.h
--- a/Engine/Plugins/Media/MediaCompositing/Source/MediaCompositing/Public/MovieSceneMediaSection.h
+++ b/Engine/Plugins/Media/MediaCompositing/Source/MediaCompositing/Public/MovieSceneMediaSection.h
@@ -73,6 +73,14 @@ public:
 	UPROPERTY(EditAnywhere, BlueprintReadWrite, Category = "Media", meta = (EditCondition = "bHasMediaPlayerProxy", HideEditConditionToggle = true, EditConditionHides), AdvancedDisplay)
 	int32 TextureIndex;
 
+	/** If true, the specified framerate alignment will be used instead of the media's frame rate. Use this in case the automatic media frame rate is inaccurate. */
+	UPROPERTY(EditAnywhere, BlueprintReadWrite, Category="Media", meta = (InlineEditConditionToggle), AdvancedDisplay)
+	bool bManualFrameRateAlignment;
+
+	/** Specify the frame rate for internal presentation frame alignment (i.e. frame duration used for Block On Range calculation). */
+	UPROPERTY(EditAnywhere, BlueprintReadWrite, Category = "Media", meta = (EditCondition = "bManualFrameRateAlignment"), AdvancedDisplay)
+	FFrameRate FrameRateAlignment;
+	
 	/** True if the object bound to this track has a media player proxy. */
 	UPROPERTY(Transient)
 	bool bHasMediaPlayerProxy;
diff --git a/Engine/Plugins/Media/MediaCompositing/Source/MediaCompositingEditor/Private/Sequencer/MediaThumbnailSection.cpp b/Engine/Plugins/Media/MediaCompositing/Source/MediaCompositingEditor/Private/Sequencer/MediaThumbnailSection.cpp
--- a/Engine/Plugins/Media/MediaCompositing/Source/MediaCompositingEditor/Private/Sequencer/MediaThumbnailSection.cpp
+++ b/Engine/Plugins/Media/MediaCompositing/Source/MediaCompositingEditor/Private/Sequencer/MediaThumbnailSection.cpp
@@ -790,11 +790,29 @@ void FMediaThumbnailSection::UpdateCachedMediaInfo(const UMediaPlayer* InMediaPl
 			PlayerInfo.Appendf(TEXT("%s%dx%d"), PlayerInfo.IsEmpty() ? TEXT("") : TEXT(", "), Resolution.X, Resolution.Y);
 		}
 
-		float FrameRate = InMediaPlayer->GetVideoTrackFrameRate(SelectedTrackIndex, SelectedFormatIndex);
-		if (FrameRate > 0.0f)
+		FText FrameRateInfo;
+		const float FrameRate = InMediaPlayer->GetVideoTrackFrameRate(SelectedTrackIndex, SelectedFormatIndex);
+
+		UMovieSceneMediaSection* MediaSection = CastChecked<UMovieSceneMediaSection>(Section);
+		if (MediaSection && MediaSection->bManualFrameRateAlignment)
 		{
-			PlayerInfo.Appendf(TEXT("%s%s"), PlayerInfo.IsEmpty() ? TEXT("") : TEXT(", "),
-				*FText::Format(LOCTEXT("FrameRateUnits", "{0} fps"), FrameRate).ToString());
+			if (FrameRate > 0.0f)
+			{
+				FrameRateInfo = FText::Format(LOCTEXT("FrameRateUnits_WithManual", "{0} ({1}) fps"), FrameRate, MediaSection->FrameRateAlignment.AsDecimal());
+			}
+			else
+			{
+				FrameRateInfo = FText::Format(LOCTEXT("FrameRateUnits_ManualOnly", "({0}) fps"), MediaSection->FrameRateAlignment.AsDecimal());
+			}
+		}
+		else if (FrameRate > 0.0f)
+		{
+			FrameRateInfo = FText::Format(LOCTEXT("FrameRateUnits", "{0} fps"), FrameRate);
+		}
+
+		if (!FrameRateInfo.IsEmpty())
+		{
+			PlayerInfo.Appendf(TEXT("%s%s"), PlayerInfo.IsEmpty() ? TEXT("") : TEXT(", "), *FrameRateInfo.ToString());
 		}
 	}
 
