[MediaCompositing] Sequencer Media Track: Fix crash on exit when running a sequencer in game mode.

diff --git a/Engine/Plugins/Media/MediaCompositing/Source/MediaCompositing/Private/MovieScene/MovieSceneMediaPlayerStore.cpp b/Engine/Plugins/Media/MediaCompositing/Source/MediaCompositing/Private/MovieScene/MovieSceneMediaPlayerStore.cpp
--- a/Engine/Plugins/Media/MediaCompositing/Source/MediaCompositing/Private/MovieScene/MovieSceneMediaPlayerStore.cpp
+++ b/Engine/Plugins/Media/MediaCompositing/Source/MediaCompositing/Private/MovieScene/MovieSceneMediaPlayerStore.cpp
@@ -14,11 +14,13 @@ namespace UE::MediaCompositing::Private
 FMovieSceneMediaPlayerStore::FMovieSceneMediaPlayerStore()
 {
 	OnEndFrameDelegateHandle = FCoreDelegates::OnEndFrame.AddRaw(this, &FMovieSceneMediaPlayerStore::OnEndFrame);
+	OnPreExitDelegateHandle = FCoreDelegates::OnPreExit.AddRaw(this, &FMovieSceneMediaPlayerStore::OnPreExit);
 }
 
 FMovieSceneMediaPlayerStore::~FMovieSceneMediaPlayerStore()
 {
 	FCoreDelegates::OnEndFrame.Remove(OnEndFrameDelegateHandle);
+	FCoreDelegates::OnPreExit.Remove(OnPreExitDelegateHandle);
 	CloseRemainingPlayers();
 }
 
@@ -91,6 +93,11 @@ void FMovieSceneMediaPlayerStore::OnEndFrame()
 	CloseRemainingPlayers();
 }
 
+void FMovieSceneMediaPlayerStore::OnPreExit()
+{
+	CloseRemainingPlayers();
+}
+
 void FMovieSceneMediaPlayerStore::CloseRemainingPlayers()
 {
 	for (const TPair<FObjectKey, UMediaPlayer*>& MediaPlayerEntry : MediaPlayers)
diff --git a/Engine/Plugins/Media/MediaCompositing/Source/MediaCompositing/Private/MovieScene/MovieSceneMediaPlayerStore.h b/Engine/Plugins/Media/MediaCompositing/Source/MediaCompositing/Private/MovieScene/MovieSceneMediaPlayerStore.h
--- a/Engine/Plugins/Media/MediaCompositing/Source/MediaCompositing/Private/MovieScene/MovieSceneMediaPlayerStore.h
+++ b/Engine/Plugins/Media/MediaCompositing/Source/MediaCompositing/Private/MovieScene/MovieSceneMediaPlayerStore.h
@@ -106,6 +106,9 @@ private:
 	/** Delegate handler for OnEndFrame. This is used to hook up CloseRemainingPlayers() somewhere after the end of the sequence evaluation. */
 	void OnEndFrame();
 
+	/** Delegate handler for OnPreExit. This is used to hook up CloseRemainingPlayers() before the final GC in case the application exits. */
+	void OnPreExit();
+
 	/** Called after the evaluation is completed to close all non-acquired players. */
 	void CloseRemainingPlayers();
 
@@ -116,6 +119,7 @@ private:
 	TMap<FObjectKey, FMovieSceneMediaPlayerPriority> MediaPlayerPriorities;
 
 	FDelegateHandle OnEndFrameDelegateHandle;
+	FDelegateHandle OnPreExitDelegateHandle;
 };
 
 /**
