Media Compositing: When adding a media section to sequencer, it will now ask if you wish to pick or create a media texture, skip setting one or disable the prompt.

#Synced-CL 48081101

diff --git a/Engine/Plugins/Media/MediaCompositing/Source/MediaCompositing/Private/MovieScene/MovieSceneMediaTrack.cpp b/Engine/Plugins/Media/MediaCompositing/Source/MediaCompositing/Private/MovieScene/MovieSceneMediaTrack.cpp
--- a/Engine/Plugins/Media/MediaCompositing/Source/MediaCompositing/Private/MovieScene/MovieSceneMediaTrack.cpp
+++ b/Engine/Plugins/Media/MediaCompositing/Source/MediaCompositing/Private/MovieScene/MovieSceneMediaTrack.cpp
@@ -118,8 +118,7 @@ UMovieSceneMediaSection* UMovieSceneMediaTrack::AddNewSectionOnRow(FFrameNumber
 
 	NewSection->InitialPlacementOnRow(MediaSections, Time, DurationToUse.FrameNumber.Value, RowIndex);
 
-	MediaSections.Add(NewSection);
-	UpdateSectionTextureIndices();
+	UMovieSceneTrack::AddSection(FSectionParameter(*NewSection));
 
 	return NewSection;
 }
diff --git a/Engine/Plugins/Media/MediaCompositing/Source/MediaCompositingEditor/MediaCompositingEditor.Build.cs b/Engine/Plugins/Media/MediaCompositing/Source/MediaCompositingEditor/MediaCompositingEditor.Build.cs
--- a/Engine/Plugins/Media/MediaCompositing/Source/MediaCompositingEditor/MediaCompositingEditor.Build.cs
+++ b/Engine/Plugins/Media/MediaCompositing/Source/MediaCompositingEditor/MediaCompositingEditor.Build.cs
@@ -9,8 +9,11 @@ namespace UnrealBuildTool.Rules
 			PrivateDependencyModuleNames.AddRange(
 				new string[]
 				{
+					"AssetRegistry",
+					"ContentBrowser",
 					"Core",
 					"CoreUObject",
+					"DeveloperSettings",
 					"EditorFramework",
 					"EditorWidgets",
 					"Engine",
diff --git a/Engine/Plugins/Media/MediaCompositing/Source/MediaCompositingEditor/Private/Sequencer/MediaTrackEditor.cpp b/Engine/Plugins/Media/MediaCompositing/Source/MediaCompositingEditor/Private/Sequencer/MediaTrackEditor.cpp
--- a/Engine/Plugins/Media/MediaCompositing/Source/MediaCompositingEditor/Private/Sequencer/MediaTrackEditor.cpp
+++ b/Engine/Plugins/Media/MediaCompositing/Source/MediaCompositingEditor/Private/Sequencer/MediaTrackEditor.cpp
@@ -23,8 +23,8 @@
 #include "Misc/QualifiedFrameTime.h"
 #include "LevelSequence.h"
 #include "SequencerSettings.h"
-
 #include "Sequencer/MediaThumbnailSection.h"
+#include "Utils/MovieSceneMediaSectionHelpers.h"
 
 
 #define LOCTEXT_NAMESPACE "FMediaTrackEditor"
@@ -318,6 +318,10 @@ void FMediaTrackEditor::AddNewSection(const FAssetData& AssetData, UMovieSceneMe
 			GetSequencer()->ThrobSectionSelection();
 
 			GetSequencer()->NotifyMovieSceneDataChanged(EMovieSceneDataChangeType::MovieSceneStructureItemAdded);
+
+			UE::MediaCompositingEditor::Private::FMovieSceneMediaSectionHelpers::ShowMediaTexturePrompt(
+				Cast<UMovieSceneMediaSection>(NewSection)
+			);
 		}
 	}
 }
diff --git a/Engine/Plugins/Media/MediaCompositing/Source/MediaCompositingEditor/Private/Utils/MediaCompositingEditorSettings.cpp b/Engine/Plugins/Media/MediaCompositing/Source/MediaCompositingEditor/Private/Utils/MediaCompositingEditorSettings.cpp
new file mode 100644
--- /dev/null
+++ b/Engine/Plugins/Media/MediaCompositing/Source/MediaCompositingEditor/Private/Utils/MediaCompositingEditorSettings.cpp
@@ -0,0 +1,3 @@
+// Copyright Epic Games, Inc. All Rights Reserved.
+
+#include "Utils/MediaCompositingEditorSettings.h"
\ No newline at end of file
diff --git a/Engine/Plugins/Media/MediaCompositing/Source/MediaCompositingEditor/Private/Utils/MediaCompositingEditorSettings.h b/Engine/Plugins/Media/MediaCompositing/Source/MediaCompositingEditor/Private/Utils/MediaCompositingEditorSettings.h
new file mode 100644
--- /dev/null
+++ b/Engine/Plugins/Media/MediaCompositing/Source/MediaCompositingEditor/Private/Utils/MediaCompositingEditorSettings.h
@@ -0,0 +1,17 @@
+// Copyright Epic Games, Inc. All Rights Reserved.
+
+#pragma once
+
+#include "Engine/DeveloperSettings.h"
+
+#include "MediaCompositingEditorSettings.generated.h"
+
+UCLASS(Config = EditorPerProjectUserSettings, meta = (DisplayName = "Media Compositing Editor Settings"))
+class UMediaCompositingEditorSettings : public UDeveloperSettings
+{
+	GENERATED_BODY()
+
+public:
+	UPROPERTY(EditAnywhere, Config, Category="Media Compositing")
+	bool bDisplayMediaTexturePrompt = true;
+};
diff --git a/Engine/Plugins/Media/MediaCompositing/Source/MediaCompositingEditor/Private/Utils/MovieSceneMediaSectionHelpers.cpp b/Engine/Plugins/Media/MediaCompositing/Source/MediaCompositingEditor/Private/Utils/MovieSceneMediaSectionHelpers.cpp
new file mode 100644
--- /dev/null
+++ b/Engine/Plugins/Media/MediaCompositing/Source/MediaCompositingEditor/Private/Utils/MovieSceneMediaSectionHelpers.cpp
@@ -0,0 +1,102 @@
+// Copyright Epic Games, Inc. All Rights Reserved.
+
+#include "Utils/MovieSceneMediaSectionHelpers.h"
+
+#include "MediaTexture.h"
+#include "Misc/MessageDialog.h"
+#include "MovieSceneMediaSection.h"
+#include "MovieSceneMediaTrack.h"
+#include "MovieSceneTrack.h"
+#include "Utils/MediaCompositingEditorSettings.h"
+#include "Utils/SMovieSceneMediaSectionMediaTexturePrompt.h"
+
+#define LOCTEXT_NAMESPACE "MovieSceneMediaSectionHelpers"
+
+namespace UE::MediaCompositingEditor::Private
+{
+
+void FMovieSceneMediaSectionHelpers::ShowMediaTexturePrompt(UMovieSceneMediaSection* InMediaSection)
+{
+	if (!InMediaSection)
+	{
+		return;
+	}
+
+	if (!GetDefault<UMediaCompositingEditorSettings>()->bDisplayMediaTexturePrompt)
+	{
+		return;
+	}
+
+	if (InMediaSection->MediaTexture || InMediaSection->bUseExternalMediaPlayer
+		|| InMediaSection->bHasMediaPlayerProxy)
+	{
+		return;
+	}
+
+	TArray<TWeakObjectPtr<UMediaTexture>> UsedMediaTextures;
+
+	if (UMovieSceneMediaTrack* Track = Cast<UMovieSceneMediaTrack>(InMediaSection->GetOuter()))
+	{
+		const TArray<UMovieSceneSection*>& Sections = Track->GetAllSections();
+		UsedMediaTextures.Reserve(Sections.Num() - 1);
+
+		for (UMovieSceneSection* Section : Sections)
+		{
+			if (Section == InMediaSection)
+			{
+				continue;
+			}
+
+			if (Section->GetRowIndex() != InMediaSection->GetRowIndex())
+			{
+				continue;
+			}
+
+			if (UMovieSceneMediaSection* MediaSection = Cast<UMovieSceneMediaSection>(Section))
+			{
+				if (MediaSection->MediaTexture)
+				{
+					UsedMediaTextures.AddUnique(MediaSection->MediaTexture);
+				}
+			}
+		}
+	}
+
+	TSharedRef<SMovieSceneMediaSectionMediaTexturePrompt> Widget = SMovieSceneMediaSectionMediaTexturePrompt::PromptUser(UsedMediaTextures);
+
+	switch (Widget->GetResponse())
+	{
+		case EMovieSceneMediaSectionMediaTexturePromptResponse::CreateAsset:
+		case EMovieSceneMediaSectionMediaTexturePromptResponse::SelectAsset:
+			SetMediaTexture(InMediaSection, Widget->GetMediaTexture());
+			break;
+
+		case EMovieSceneMediaSectionMediaTexturePromptResponse::Skip:
+			break;
+
+		case EMovieSceneMediaSectionMediaTexturePromptResponse::Disable:
+			DisablePrompt();
+			break;
+	}
+}
+
+void FMovieSceneMediaSectionHelpers::DisablePrompt()
+{
+	UMediaCompositingEditorSettings* Settings = GetMutableDefault<UMediaCompositingEditorSettings>();
+	Settings->bDisplayMediaTexturePrompt = false;
+	Settings->SaveConfig();
+
+	FMessageDialog::Open(EAppMsgType::Ok, LOCTEXT("HowToTurnBackOn", "You can re-enable this prompt in the Editor Settings in the Media Compositing Editor section."));
+}
+
+void FMovieSceneMediaSectionHelpers::SetMediaTexture(UMovieSceneMediaSection* InSection, UMediaTexture* InTexture)
+{
+	if (InSection && InTexture)
+	{
+		InSection->MediaTexture = InTexture;
+	}
+}
+
+}
+
+#undef LOCTEXT_NAMESPACE
diff --git a/Engine/Plugins/Media/MediaCompositing/Source/MediaCompositingEditor/Private/Utils/SMovieSceneMediaSectionMediaTexturePrompt.cpp b/Engine/Plugins/Media/MediaCompositing/Source/MediaCompositingEditor/Private/Utils/SMovieSceneMediaSectionMediaTexturePrompt.cpp
new file mode 100644
--- /dev/null
+++ b/Engine/Plugins/Media/MediaCompositing/Source/MediaCompositingEditor/Private/Utils/SMovieSceneMediaSectionMediaTexturePrompt.cpp
@@ -0,0 +1,350 @@
+// Copyright Epic Games, Inc. All Rights Reserved.
+
+#include "Utils/SMovieSceneMediaSectionMediaTexturePrompt.h"
+
+#include "AssetRegistry/AssetRegistryModule.h"
+#include "Containers/UnrealString.h"
+#include "ContentBrowserModule.h"
+#include "Framework/Application/SlateApplication.h"
+#include "Framework/MultiBox/MultiBoxBuilder.h"
+#include "IContentBrowserSingleton.h"
+#include "MediaTexture.h"
+#include "Misc/MessageDialog.h"
+#include "Misc/Paths.h"
+#include "Modules/ModuleManager.h"
+#include "Styling/AppStyle.h"
+#include "Widgets/Images/SImage.h"
+#include "Widgets/Input/SButton.h"
+#include "Widgets/Input/SHyperlink.h"
+#include "Widgets/SBoxPanel.h"
+#include "Widgets/SWindow.h"
+#include "Widgets/Text/STextBlock.h"
+#include "UObject/Object.h"
+#include "UObject/Package.h"
+#include "UObject/UObjectGlobals.h"
+
+#define LOCTEXT_NAMESPACE "SMovieSceneMediaSectionMediaTexturePrompt"
+
+namespace UE::MediaCompositingEditor::Private
+{
+
+TSharedRef<SMovieSceneMediaSectionMediaTexturePrompt> SMovieSceneMediaSectionMediaTexturePrompt::PromptUser(TArrayView<TWeakObjectPtr<UMediaTexture>> InUsedTexturesWeak)
+{
+	TSharedRef<SMovieSceneMediaSectionMediaTexturePrompt> PromptWidget = SNew(SMovieSceneMediaSectionMediaTexturePrompt, InUsedTexturesWeak);
+
+	TSharedRef<SWindow> Window = SNew(SWindow)
+		.Title(LOCTEXT("WindowTitle", "Add Media Texture"))
+		.SizingRule(ESizingRule::Autosized)
+		.AutoCenter(EAutoCenter::PreferredWorkArea)
+		.SupportsMinimize(false)
+		.SupportsMaximize(false)
+		.IsTopmostWindow(false)
+		.UserResizeBorder(false)
+		.Content()
+		[
+			PromptWidget
+		];
+
+	PromptWidget->Window = Window;
+	
+	FSlateApplication::Get().AddModalWindow(Window, nullptr);
+
+	return PromptWidget;
+}
+
+void SMovieSceneMediaSectionMediaTexturePrompt::Construct(const FArguments& InArgs, TArrayView<TWeakObjectPtr<UMediaTexture>> InUsedTexturesWeak)
+{
+	UsedTexturesWeak = InUsedTexturesWeak;
+
+	TSharedRef<SVerticalBox> Contents =	SNew(SVerticalBox)
+		
+		+ SVerticalBox::Slot()
+		.AutoHeight()
+		[
+			SNew(SHorizontalBox)
+			
+			+ SHorizontalBox::Slot()
+			.MinWidth(32.f)
+			.MaxWidth(32.f)
+			.VAlign(VAlign_Center)
+			[
+				SNew(SImage)
+				.DesiredSizeOverride(FVector2D(32.f))
+				.Image(FAppStyle::GetBrush("ClassIcon.MediaTexture"))
+			]
+
+			+ SHorizontalBox::Slot()
+			.AutoWidth()
+			.HAlign(HAlign_Left)
+			.VAlign(VAlign_Center)
+			.Padding(10.f, 0.f, 0.f, 0.f)
+			[
+				SNew(STextBlock)
+				.Text(LOCTEXT("AddMediaTexture", "Add a Media Texture to the new section?"))
+			]
+		]
+
+		+ SVerticalBox::Slot()
+		.AutoHeight()
+		.Padding(0.f, 30.f, 0.f, 0.f)
+		[
+			SNew(SHorizontalBox)
+			
+			+ SHorizontalBox::Slot()
+			.AutoWidth()
+			[
+				SNew(SButton)
+				.Text(LOCTEXT("CreateAsset", "Create Texture"))
+				.OnClicked(this, &SMovieSceneMediaSectionMediaTexturePrompt::OnCreateNewClicked)
+				.ContentPadding(5.f)
+			]
+			
+			+ SHorizontalBox::Slot()
+			.AutoWidth()
+			.Padding(10.f, 0.f, 0.f, 0.f)
+			[
+				SNew(SButton)
+				.Text(LOCTEXT("SelectAsset", "Pick Texture"))
+				.OnClicked(this, &SMovieSceneMediaSectionMediaTexturePrompt::OnSelectAssetClicked)
+				.ContentPadding(5.f)
+			]
+			
+			+ SHorizontalBox::Slot()
+			.AutoWidth()
+			.Padding(10.f, 0.f, 0.f, 0.f)
+			[
+				SNew(SButton)
+				.Text(LOCTEXT("Skip", "Skip"))
+				.OnClicked(this, &SMovieSceneMediaSectionMediaTexturePrompt::OnSkipClicked)
+				.ContentPadding(5.f)
+			]
+			
+			+ SHorizontalBox::Slot()
+			.AutoWidth()
+			.Padding(10.f, 0.f, 0.f, 0.f)
+			[
+				SNew(SButton)
+				.Text(LOCTEXT("Diable", "Don't Ask Again"))
+				.OnClicked(this, &SMovieSceneMediaSectionMediaTexturePrompt::OnDisableClicked)
+				.ContentPadding(5.f)
+			]
+		];
+
+	if (!UsedTexturesWeak.IsEmpty())
+	{
+		bool bHasValidTexture = false;
+
+		for (const TWeakObjectPtr<UMediaTexture> MediaTextureWeak : UsedTexturesWeak)
+		{
+			if (UMediaTexture* MediaTexture = MediaTextureWeak.Get())
+			{
+				bHasValidTexture = true;
+				break;
+			}
+		}
+		
+		if (bHasValidTexture)
+		{
+			Contents->AddSlot()
+				.AutoHeight()
+				.Padding(0.f, 30.f, 0.f, 0.f)
+				[
+					SNew(STextBlock)
+					.Text(LOCTEXT("TexturesAlreadyUsed", "Textures already on this row:"))
+				];
+
+			for (const TWeakObjectPtr<UMediaTexture> MediaTextureWeak : UsedTexturesWeak)
+			{
+				if (UMediaTexture* MediaTexture = MediaTextureWeak.Get())
+				{
+					Contents->AddSlot()
+						.AutoHeight()
+						.Padding(0.f, 10.f, 0.f, 0.f)
+						.HAlign(HAlign_Left)
+						[
+							SNew(SHyperlink)
+							.Text(FText::FromString(MediaTexture->GetPathName()))
+							.OnNavigate(this, &SMovieSceneMediaSectionMediaTexturePrompt::OnUsedTexturePicked, MediaTextureWeak)
+						];
+				}
+			}
+
+			Contents->AddSlot()
+				.AutoHeight()
+				.Padding(0.f, 20.f, 0.f, 0.f)
+				[
+					SNew(SHorizontalBox)
+			
+					+ SHorizontalBox::Slot()
+					.MinWidth(16.f)
+					.MaxWidth(16.f)
+					.VAlign(VAlign_Center)
+					[
+						SNew(SImage)
+						.DesiredSizeOverride(FVector2D(16.f))
+						.Image(FAppStyle::GetBrush("Icons.Warning.Solid"))
+					]
+
+					+ SHorizontalBox::Slot()
+					.AutoWidth()
+					.HAlign(HAlign_Left)
+					.VAlign(VAlign_Center)
+					.Padding(10.f, 0.f, 0.f, 0.f)
+					[
+						SNew(STextBlock)
+						.Text(LOCTEXT("ReuseTextureWarning", "Sections with the same texture must not overlap!"))
+					]
+				];
+		}
+	}
+
+	ChildSlot
+	.Padding(20.f)
+	[
+		Contents
+	];
+}
+
+EMovieSceneMediaSectionMediaTexturePromptResponse SMovieSceneMediaSectionMediaTexturePrompt::GetResponse() const
+{
+	return Response;
+}
+
+UMediaTexture* SMovieSceneMediaSectionMediaTexturePrompt::GetMediaTexture() const
+{
+	return Texture.Get();
+}
+
+FReply SMovieSceneMediaSectionMediaTexturePrompt::OnCreateNewClicked()
+{
+	ShowAssetCreator();
+
+	return FReply::Handled();
+}
+
+FReply SMovieSceneMediaSectionMediaTexturePrompt::OnSelectAssetClicked()
+{
+	ShowAssetPicker();
+
+	return FReply::Handled();
+}
+
+FReply SMovieSceneMediaSectionMediaTexturePrompt::OnSkipClicked()
+{
+	Response = EMovieSceneMediaSectionMediaTexturePromptResponse::Skip;
+	RequestClose();
+
+	return FReply::Handled();
+}
+
+FReply SMovieSceneMediaSectionMediaTexturePrompt::OnDisableClicked()
+{
+	Response = EMovieSceneMediaSectionMediaTexturePromptResponse::Disable;
+	RequestClose();
+
+	return FReply::Handled();
+}
+
+void SMovieSceneMediaSectionMediaTexturePrompt::RequestClose()
+{
+	if (Window.IsValid())
+	{
+		FSlateApplication::Get().RequestDestroyWindow(Window.ToSharedRef());
+	}
+}
+
+void SMovieSceneMediaSectionMediaTexturePrompt::ShowAssetCreator()
+{
+	FSaveAssetDialogConfig SaveAssetDialogConfig;
+	SaveAssetDialogConfig.DialogTitleOverride = LOCTEXT("CreateAssetDialogTitle", "Save Media Texture");
+	SaveAssetDialogConfig.DefaultPath = "/Game";
+	SaveAssetDialogConfig.DefaultAssetName = "MediaTexture";
+	SaveAssetDialogConfig.ExistingAssetPolicy = ESaveAssetDialogExistingAssetPolicy::Disallow;
+
+	FContentBrowserModule& ContentBrowserModule = FModuleManager::LoadModuleChecked<FContentBrowserModule>("ContentBrowser");
+	const FString SaveObjectPath = ContentBrowserModule.Get().CreateModalSaveAssetDialog(SaveAssetDialogConfig);
+
+	if (SaveObjectPath.IsEmpty())
+	{
+		return;
+	}
+
+	const FString PackagePath = FPaths::GetBaseFilename(*SaveObjectPath, false);
+
+	UPackage* Package = CreatePackage(*PackagePath);
+
+	if (!Package)
+	{
+		return;
+	}
+
+	const FString AssetName = FPaths::GetBaseFilename(*SaveObjectPath, true);
+
+	UMediaTexture* MediaTexture = NewObject<UMediaTexture>(Package, *AssetName, RF_Public | RF_Standalone | RF_Transactional);
+
+	if (!MediaTexture)
+	{
+		return;
+	}
+
+	FAssetRegistryModule::AssetCreated(MediaTexture);
+	Package->MarkPackageDirty();
+
+	Texture.Reset(MediaTexture);
+	Response = EMovieSceneMediaSectionMediaTexturePromptResponse::CreateAsset;
+
+	RequestClose();
+}
+
+void SMovieSceneMediaSectionMediaTexturePrompt::ShowAssetPicker()
+{
+	FOpenAssetDialogConfig OpenAssetDialogConfig;
+	OpenAssetDialogConfig.DialogTitleOverride = LOCTEXT("OpenAssetDialogTitle", "Pick Media Texture");
+	OpenAssetDialogConfig.DefaultPath = "/Game";
+	OpenAssetDialogConfig.AssetClassNames.Add(UMediaTexture::StaticClass()->GetClassPathName());
+	OpenAssetDialogConfig.bAllowMultipleSelection = false;
+
+	FContentBrowserModule& ContentBrowserModule = FModuleManager::LoadModuleChecked<FContentBrowserModule>("ContentBrowser");
+	TArray<FAssetData> PickedAssets = ContentBrowserModule.Get().CreateModalOpenAssetDialog(OpenAssetDialogConfig);
+
+	if (PickedAssets.Num() != 1)
+	{
+		return;
+	}
+
+	UMediaTexture* MediaTexture = Cast<UMediaTexture>(PickedAssets[0].GetAsset());
+
+	if (!MediaTexture)
+	{
+		return;
+	}
+
+	Texture.Reset(MediaTexture);
+	Response = EMovieSceneMediaSectionMediaTexturePromptResponse::SelectAsset;
+
+	RequestClose();
+}
+
+void SMovieSceneMediaSectionMediaTexturePrompt::OnUsedTexturePicked(TWeakObjectPtr<UMediaTexture> InPickedTextureWeak)
+{
+	UMediaTexture* PickedTexture = InPickedTextureWeak.Get();
+
+	if (!PickedTexture)
+	{
+		FMessageDialog::Open(
+			EAppMsgCategory::Error,
+			EAppMsgType::Ok, 
+			LOCTEXT("TextureNoLongerExists", "That texture no longer exists.")
+		);
+	}
+	else
+	{
+		Texture.Reset(PickedTexture);
+		Response = EMovieSceneMediaSectionMediaTexturePromptResponse::SelectAsset;
+		RequestClose();
+	}
+}
+
+}
+
+#undef LOCTEXT_NAMESPACE
diff --git a/Engine/Plugins/Media/MediaCompositing/Source/MediaCompositingEditor/Private/Utils/SMovieSceneMediaSectionMediaTexturePrompt.h b/Engine/Plugins/Media/MediaCompositing/Source/MediaCompositingEditor/Private/Utils/SMovieSceneMediaSectionMediaTexturePrompt.h
new file mode 100644
--- /dev/null
+++ b/Engine/Plugins/Media/MediaCompositing/Source/MediaCompositingEditor/Private/Utils/SMovieSceneMediaSectionMediaTexturePrompt.h
@@ -0,0 +1,61 @@
+// Copyright Epic Games, Inc. All Rights Reserved.
+
+#pragma once
+
+#include "Widgets/SCompoundWidget.h"
+
+#include "UObject/StrongObjectPtr.h"
+
+class SWindow;
+class UMediaTexture;
+
+namespace UE::MediaCompositingEditor::Private
+{
+
+enum class EMovieSceneMediaSectionMediaTexturePromptResponse : uint8
+{
+	CreateAsset,
+	SelectAsset,
+	Skip,
+	Disable
+};
+
+class SMovieSceneMediaSectionMediaTexturePrompt : public SCompoundWidget
+{
+public:
+	SLATE_BEGIN_ARGS(SMovieSceneMediaSectionMediaTexturePrompt) 
+		{}
+	SLATE_END_ARGS()
+
+	static TSharedRef<SMovieSceneMediaSectionMediaTexturePrompt> PromptUser(TArrayView<TWeakObjectPtr<UMediaTexture>> InUsedTexturesWeak);
+
+	void Construct(const FArguments& InArgs, TArrayView<TWeakObjectPtr<UMediaTexture>> InUsedTexturesWeak);
+
+	EMovieSceneMediaSectionMediaTexturePromptResponse GetResponse() const;
+
+	UMediaTexture* GetMediaTexture() const;
+
+protected:
+	TArrayView<TWeakObjectPtr<UMediaTexture>> UsedTexturesWeak;
+	TSharedPtr<SWindow> Window;
+	EMovieSceneMediaSectionMediaTexturePromptResponse Response = EMovieSceneMediaSectionMediaTexturePromptResponse::Skip;
+	TStrongObjectPtr<UMediaTexture> Texture;
+
+	FReply OnCreateNewClicked();
+
+	FReply OnSelectAssetClicked();
+
+	FReply OnSkipClicked();
+
+	FReply OnDisableClicked();
+	
+	void RequestClose();
+
+	void ShowAssetCreator();
+
+	void ShowAssetPicker();
+
+	void OnUsedTexturePicked(TWeakObjectPtr<UMediaTexture> InPickedTextureWeak);
+};
+
+}
\ No newline at end of file
diff --git a/Engine/Plugins/Media/MediaCompositing/Source/MediaCompositingEditor/Public/Utils/MovieSceneMediaSectionHelpers.h b/Engine/Plugins/Media/MediaCompositing/Source/MediaCompositingEditor/Public/Utils/MovieSceneMediaSectionHelpers.h
new file mode 100644
--- /dev/null
+++ b/Engine/Plugins/Media/MediaCompositing/Source/MediaCompositingEditor/Public/Utils/MovieSceneMediaSectionHelpers.h
@@ -0,0 +1,29 @@
+// Copyright Epic Games, Inc. All Rights Reserved.
+
+#pragma once
+
+#include "HAL/Platform.h"
+
+class UMediaTexture;
+class UMovieSceneMediaSection;
+
+#define UE_API MEDIACOMPOSITINGEDITOR_API
+
+namespace UE::MediaCompositingEditor::Private
+{
+
+struct FMovieSceneMediaSectionHelpers
+{
+	static UE_API void ShowMediaTexturePrompt(UMovieSceneMediaSection* InMediaSection);
+
+private:
+	/** Disables the prompt for the future. */
+	static void DisablePrompt();
+
+	/** Callback to set the media texture on to the section. */
+	static void SetMediaTexture(UMovieSceneMediaSection* InSection, UMediaTexture* InTexture);
+};
+
+}
+
+#undef UE_API
