Media: Add missing modify when dropping media into a track

This fixes a crash on undo

#Synced-CL 50777405

diff --git a/Engine/Plugins/Media/MediaCompositing/Source/MediaCompositingEditor/Private/Sequencer/MediaTrackEditor.cpp b/Engine/Plugins/Media/MediaCompositing/Source/MediaCompositingEditor/Private/Sequencer/MediaTrackEditor.cpp
--- a/Engine/Plugins/Media/MediaCompositing/Source/MediaCompositingEditor/Private/Sequencer/MediaTrackEditor.cpp
+++ b/Engine/Plugins/Media/MediaCompositing/Source/MediaCompositingEditor/Private/Sequencer/MediaTrackEditor.cpp
@@ -256,6 +256,8 @@ FKeyPropertyResult FMediaTrackEditor::AddAttachedMediaSource(FFrameNumber KeyTim
 			if (ensure(Track))
 			{
 				auto MediaTrack = Cast<UMovieSceneMediaTrack>(Track);
+				MediaTrack->Modify();
+
 				UMovieSceneSection* NewSection = MediaTrack->AddNewMediaSourceOnRow(*MediaSource, KeyTime, RowIndex);
 				MediaTrack->SetDisplayName(LOCTEXT("MediaTrackName", "Media"));
 				KeyPropertyResult.bTrackModified = true;
@@ -279,6 +281,7 @@ FKeyPropertyResult FMediaTrackEditor::AddMediaSource(FFrameNumber KeyTime, UMedi
 	FFindOrCreateRootTrackResult<UMovieSceneMediaTrack> TrackResult = FindOrCreateRootTrack<UMovieSceneMediaTrack>();
 	UMovieSceneTrack* Track = TrackResult.Track;
 	auto MediaTrack = Cast<UMovieSceneMediaTrack>(Track);
+	MediaTrack->Modify();
 
 	UMovieSceneSection* NewSection = MediaTrack->AddNewMediaSourceOnRow(*MediaSource, KeyTime, RowIndex);
 
