[MediaCompositing] Fix Media Texture not updating when created through the media track creation prompt.

#Synced-CL 49891360

diff --git a/Engine/Plugins/Media/MediaCompositing/Source/MediaCompositingEditor/Private/Utils/SMovieSceneMediaSectionMediaTexturePrompt.cpp b/Engine/Plugins/Media/MediaCompositing/Source/MediaCompositingEditor/Private/Utils/SMovieSceneMediaSectionMediaTexturePrompt.cpp
--- a/Engine/Plugins/Media/MediaCompositing/Source/MediaCompositingEditor/Private/Utils/SMovieSceneMediaSectionMediaTexturePrompt.cpp
+++ b/Engine/Plugins/Media/MediaCompositing/Source/MediaCompositingEditor/Private/Utils/SMovieSceneMediaSectionMediaTexturePrompt.cpp
@@ -290,6 +290,8 @@ void SMovieSceneMediaSectionMediaTexturePrompt::ShowAssetCreator()
 	FAssetRegistryModule::AssetCreated(MediaTexture);
 	Package->MarkPackageDirty();
 
+	MediaTexture->UpdateResource();
+
 	Texture.Reset(MediaTexture);
 	Response = EMovieSceneMediaSectionMediaTexturePromptResponse::CreateAsset;
 
