[MediaCompositing] Build Media Source Thumbnail on demand from Asset Thumbnail when requested for the Media Track Thumbnail Section.

The Media Source Thumbnail is not serialized since it is a transient full size media texture that is only used to generate the thumbnail and is not available otherwise.

For the media track thumbnail, we thus rebuild a thumbnail texture from the object thumbnail already stored in the asset. So if the asset has a thumbnail, so will the media track.

#Synced-CL 51094493

diff --git a/Engine/Plugins/Media/MediaCompositing/Source/MediaCompositingEditor/MediaCompositingEditor.Build.cs b/Engine/Plugins/Media/MediaCompositing/Source/MediaCompositingEditor/MediaCompositingEditor.Build.cs
--- a/Engine/Plugins/Media/MediaCompositing/Source/MediaCompositingEditor/MediaCompositingEditor.Build.cs
+++ b/Engine/Plugins/Media/MediaCompositing/Source/MediaCompositingEditor/MediaCompositingEditor.Build.cs
@@ -17,6 +17,7 @@ namespace UnrealBuildTool.Rules
 					"EditorFramework",
 					"EditorWidgets",
 					"Engine",
+					"ImageCore",
 					"ImgMedia",
 					"InputCore",
 					"LevelSequence",
diff --git a/Engine/Plugins/Media/MediaCompositing/Source/MediaCompositingEditor/Private/Sequencer/MediaThumbnailSection.cpp b/Engine/Plugins/Media/MediaCompositing/Source/MediaCompositingEditor/Private/Sequencer/MediaThumbnailSection.cpp
--- a/Engine/Plugins/Media/MediaCompositing/Source/MediaCompositingEditor/Private/Sequencer/MediaThumbnailSection.cpp
+++ b/Engine/Plugins/Media/MediaCompositing/Source/MediaCompositingEditor/Private/Sequencer/MediaThumbnailSection.cpp
@@ -3,33 +3,35 @@
 #include "Sequencer/MediaThumbnailSection.h"
 
 #include "CommonRenderResources.h"
+#include "Compilation/MovieSceneCompiledDataManager.h"
+#include "Engine/Texture2D.h"
 #include "Fonts/FontMeasure.h"
 #include "Framework/Application/SlateApplication.h"
-#include "Styling/AppStyle.h"
 #include "IMediaCache.h"
 #include "IMediaTracks.h"
 #include "ISequencer.h"
+#include "ImageCoreUtils.h"
 #include "MediaPlayer.h"
 #include "MediaPlayerFacade.h"
 #include "MediaSource.h"
+#include "Misc/ObjectThumbnail.inl"
 #include "Modules/ModuleManager.h"
 #include "MovieScene.h"
+#include "MovieSceneMediaData.h"
 #include "MovieSceneMediaSection.h"
 #include "MovieSceneTimeHelpers.h"
-#include "Rendering/DrawElements.h"
+#include "ObjectTools.h"
 #include "RHI.h"
 #include "RHIStaticStates.h"
+#include "RendererInterface.h"
+#include "Rendering/DrawElements.h"
 #include "ScreenRendering.h"
 #include "SequencerSectionPainter.h"
+#include "Styling/AppStyle.h"
+#include "TimeToPixel.h"
 #include "TrackEditorThumbnail/TrackEditorThumbnailPool.h"
-#include "Compilation/MovieSceneCompiledDataManager.h"
 #include "Widgets/Layout/SBox.h"
 #include "Widgets/Text/STextBlock.h"
-#include "TimeToPixel.h"
-#include "RendererInterface.h"
-
-#include "MovieSceneMediaData.h"
-
 
 #define LOCTEXT_NAMESPACE "FMediaThumbnailSection"
 
@@ -117,6 +119,74 @@ namespace UE::MediaThumbnailSection
 
 		FColor TextColor = FColor(192, 192, 192);
 	};
+
+	/** Returns the thumbnail for the specified object or NULL if one doesn't exist yet. */
+	FObjectThumbnail* GetOrLoadThumbnailForObject(UObject* InObject)
+	{
+		FObjectThumbnail* ObjectThumbnail = ThumbnailTools::GetThumbnailForObject(InObject);
+		if (ObjectThumbnail == nullptr)
+		{
+			UPackage* ObjectPackage = InObject->GetOutermost();
+
+			// Replicating some logic from ThumbnailManager
+			const bool bMemoryPackage = FPackageName::IsMemoryPackage(InObject->GetPathName()); // Don't try to load from disk if the package is a memory package
+			const bool bUnsavedPackage = ObjectPackage->HasAnyPackageFlags(PKG_NewlyCreated);   // Don't try loading thumbnails for package that have never been saved
+
+			const bool bTryLoadThumbnailFromDisk = !bMemoryPackage && !bUnsavedPackage;
+			if (bTryLoadThumbnailFromDisk)
+			{
+				FName ObjectFullName = FName(*InObject->GetFullName());
+
+				FThumbnailMap LoadedThumbnails;
+				if (ThumbnailTools::ConditionallyLoadThumbnailsForObjects({ObjectFullName}, LoadedThumbnails))
+				{
+					ObjectThumbnail = LoadedThumbnails.Find(ObjectFullName);
+
+					// If we found one, cache it so that we don't have to do this again
+					if (ObjectThumbnail != nullptr)
+					{
+						ObjectThumbnail = ThumbnailTools::CacheThumbnail(ObjectFullName.ToString(), ObjectThumbnail, ObjectPackage);
+					}
+				}
+			}
+		}
+		return ObjectThumbnail;
+	}
+
+	/** Utility function to try to make the media source's thumbnail from the existing asset thumbnail, if any. */
+	const UTexture* TryMakeThumbnailFromAssetThumbnail(UMediaSource* InMediaSource)
+	{
+		UTexture2D* Texture = nullptr;
+		
+		const FObjectThumbnail* AssetThumbnail = GetOrLoadThumbnailForObject(InMediaSource);
+	
+		// Create texture from the asset thumbnail
+		if (AssetThumbnail && !AssetThumbnail->IsEmpty())
+		{
+			const TArray<uint8>& ImageData = AssetThumbnail->GetUncompressedImageData();
+			if (ImageData.Num() > 0)
+			{
+				Texture = UTexture2D::CreateTransient(
+					AssetThumbnail->GetImageWidth(),
+					AssetThumbnail->GetImageHeight(),
+					FImageCoreUtils::GetPixelFormatForRawImageFormat(AssetThumbnail->GetImage().Format),
+					NAME_None,
+					ImageData);
+			}
+
+			if (Texture)
+			{
+				InMediaSource->SetThumbnail(Texture);	// Store the texture directly in the media source to speed things up.
+			}
+		}
+
+		if (!AssetThumbnail || !Texture)
+		{
+			// Cache an empty thumbnail to speed up the process next time.
+			ThumbnailTools::CacheEmptyThumbnail(InMediaSource->GetFullName(), InMediaSource->GetPackage());
+		}
+		return Texture;
+	}
 }
 
 /* FMediaThumbnailSection structors
@@ -133,6 +203,16 @@ FMediaThumbnailSection::FMediaThumbnailSection(UMovieSceneMediaSection& InSectio
 		InSequencer->OnBeginScrubbingEvent().AddRaw(this, &FMediaThumbnailSection::OnBeginScrubbingEvent);
 		InSequencer->OnEndScrubbingEvent().AddRaw(this, &FMediaThumbnailSection::OnEndScrubbingEvent);
 	}
+
+	// Build the media source's thumbnail
+	if (UMediaSource* MediaSource = GetMediaSource())
+	{
+		const UTexture* Thumbnail = MediaSource->GetThumbnail();
+		if (!Thumbnail)
+		{
+			UE::MediaThumbnailSection::TryMakeThumbnailFromAssetThumbnail(MediaSource);
+		}
+	}
 }
 
 
@@ -396,19 +476,25 @@ void FMediaThumbnailSection::SlipSection(FFrameNumber SlipTime)
 
 void FMediaThumbnailSection::Draw(FTrackEditorThumbnail& TrackEditorThumbnail)
 {
-	UMediaSource* MediaSource = GetMediaSource();
-	if (MediaSource != nullptr)
+	if (UMediaSource* MediaSource = GetMediaSource())
 	{
-		UTexture* Thumbnail = MediaSource->GetThumbnail();
-		if (Thumbnail != nullptr)
+		// Use media source's thumbnail if available.
+		const UTexture* Thumbnail = MediaSource->GetThumbnail();
+		if (!Thumbnail)
+		{
+			// Fallback to the asset thumbnail if there is one.
+			Thumbnail = UE::MediaThumbnailSection::TryMakeThumbnailFromAssetThumbnail(MediaSource);
+		}
+
+		if (Thumbnail)
 		{
-			FTextureReferenceRHIRef SourceTexture = Thumbnail->TextureReference.TextureReferenceRHI;
+			const FTextureReferenceRHIRef SourceTexture = Thumbnail->TextureReference.TextureReferenceRHI;
 			if (SourceTexture.IsValid())
 			{
 				// Limit thumbnail size.
 				FIntPoint RTSize(SourceTexture->GetDesc().Extent);
-				int32 SourceMaxSize = RTSize.GetMax();
-				int32 ThumbnailSize = 256;
+				const int32 SourceMaxSize = RTSize.GetMax();
+				constexpr int32 ThumbnailSize = 256;
 				if (ThumbnailSize < SourceMaxSize)
 				{
 					RTSize = ((RTSize * ThumbnailSize) / SourceMaxSize);
@@ -416,9 +502,8 @@ void FMediaThumbnailSection::Draw(FTrackEditorThumbnail& TrackEditorThumbnail)
 
 				TrackEditorThumbnail.bIgnoreAlpha = true;
 				TrackEditorThumbnail.ResizeRenderTarget(RTSize);
-				FSlateTextureRenderTarget2DResource* RenderTarget =
-					TrackEditorThumbnail.GetRenderTarget();
-				if (RenderTarget != nullptr)
+
+				if (FSlateTextureRenderTarget2DResource* RenderTarget = TrackEditorThumbnail.GetRenderTarget())
 				{
 					CopyTexture(RenderTarget, SourceTexture);
 				}
